#ue4
* now properly find maps when the project is on a separate drive
* TTP328613

diff --git a/Engine/Source/Developer/LauncherServices/Public/Common/GameProjectHelper.h b/Engine/Source/Developer/LauncherServices/Public/Common/GameProjectHelper.h
--- a/Engine/Source/Developer/LauncherServices/Public/Common/GameProjectHelper.h
+++ b/Engine/Source/Developer/LauncherServices/Public/Common/GameProjectHelper.h
@@ -82,6 +82,12 @@ public:
 
 		IFileManager::Get().FindFilesRecursive(ProjectMapNames, *FPaths::Combine(*FPaths::RootDir(), *GameName, TEXT("Content"), TEXT("Maps")), *WildCard, true, false);
 
+		// didn't find any, let's check the base GameName just in case it is a full path
+		if (ProjectMapNames.Num() == 0)
+		{
+			IFileManager::Get().FindFilesRecursive(ProjectMapNames, *FPaths::Combine(*GameName, TEXT("Content"), TEXT("Maps")), *WildCard, true, false);
+		}
+
 		for (int32 i = 0; i < ProjectMapNames.Num(); i++)
 		{
 			Result.Add(FPaths::GetBaseFilename(ProjectMapNames[i]));
