#ue4
* fix for IPP causing deploy or package to fail when running from Game Launcher

diff --git a/Engine/Source/Developer/HTML5/HTML5TargetPlatform/Private/HTML5TargetPlatform.h b/Engine/Source/Developer/HTML5/HTML5TargetPlatform/Private/HTML5TargetPlatform.h
--- a/Engine/Source/Developer/HTML5/HTML5TargetPlatform/Private/HTML5TargetPlatform.h
+++ b/Engine/Source/Developer/HTML5/HTML5TargetPlatform/Private/HTML5TargetPlatform.h
@@ -24,6 +24,8 @@ public:
 
 	// Begin ITargetPlatform interface
 
+	virtual void EnableDeviceCheck(bool OnOff) OVERRIDE {}
+
 	virtual void GetAllDevices( TArray<ITargetDevicePtr>& OutDevices ) const OVERRIDE;
 
 	virtual ECompressionFlags GetBaseCompressionMethod( ) const OVERRIDE;
diff --git a/Engine/Source/Developer/LauncherServices/Private/Launcher/LauncherDeployCommands.h b/Engine/Source/Developer/LauncherServices/Private/Launcher/LauncherDeployCommands.h
--- a/Engine/Source/Developer/LauncherServices/Private/Launcher/LauncherDeployCommands.h
+++ b/Engine/Source/Developer/LauncherServices/Private/Launcher/LauncherDeployCommands.h
@@ -82,6 +82,20 @@ public:
 		return CommandLine;
 	}
 
+	virtual bool PreExecute(FLauncherTaskChainState& ChainState) OVERRIDE
+	{
+		// disable the device check
+		const_cast<ITargetPlatform&>(TargetPlatform).EnableDeviceCheck(false);
+		return true;
+	}
+
+	virtual bool PostExecute(FLauncherTaskChainState& ChainState) OVERRIDE
+	{
+		// disable the device check
+		const_cast<ITargetPlatform&>(TargetPlatform).EnableDeviceCheck(true);
+		return true;
+	}
+
 private:
 
 	// Holds a pointer to the device proxy to deploy to.
@@ -255,6 +269,20 @@ public:
 		return CommandLine;
 	}
 
+	virtual bool PreExecute(FLauncherTaskChainState& ChainState) OVERRIDE
+	{
+		// disable the device check
+		const_cast<ITargetPlatform&>(TargetPlatform).EnableDeviceCheck(false);
+		return true;
+	}
+
+	virtual bool PostExecute(FLauncherTaskChainState& ChainState) OVERRIDE
+	{
+		// disable the device check
+		const_cast<ITargetPlatform&>(TargetPlatform).EnableDeviceCheck(true);
+		return true;
+	}
+
 private:
 
 	// Holds a pointer to the device proxy to deploy to.
diff --git a/Engine/Source/Developer/LauncherServices/Private/Launcher/LauncherPackageCommands.h b/Engine/Source/Developer/LauncherServices/Private/Launcher/LauncherPackageCommands.h
--- a/Engine/Source/Developer/LauncherServices/Private/Launcher/LauncherPackageCommands.h
+++ b/Engine/Source/Developer/LauncherServices/Private/Launcher/LauncherPackageCommands.h
@@ -229,6 +229,20 @@ public:
 		return CommandLine;
 	}
 
+	virtual bool PreExecute(FLauncherTaskChainState& ChainState) OVERRIDE
+	{
+		// disable the device check
+		const_cast<ITargetPlatform&>(TargetPlatform).EnableDeviceCheck(false);
+		return true;
+	}
+
+	virtual bool PostExecute(FLauncherTaskChainState& ChainState) OVERRIDE
+	{
+		// disable the device check
+		const_cast<ITargetPlatform&>(TargetPlatform).EnableDeviceCheck(true);
+		return true;
+	}
+
 private:
 
 	// Holds a pointer to the target platform.
diff --git a/Engine/Source/Developer/Linux/LinuxTargetPlatform/Private/LinuxTargetPlatform.h b/Engine/Source/Developer/Linux/LinuxTargetPlatform/Private/LinuxTargetPlatform.h
--- a/Engine/Source/Developer/Linux/LinuxTargetPlatform/Private/LinuxTargetPlatform.h
+++ b/Engine/Source/Developer/Linux/LinuxTargetPlatform/Private/LinuxTargetPlatform.h
@@ -37,6 +37,8 @@ public:
 
 	// Begin ITargetPlatform interface
 
+	virtual void EnableDeviceCheck(bool OnOff) OVERRIDE {}
+
 	virtual void GetAllDevices( TArray<ITargetDevicePtr>& OutDevices ) const OVERRIDE
 	{
 		// TODO: ping all the machines in a local segment and/or try to connect to port 22 of those that respond
diff --git a/Engine/Source/Developer/TargetPlatform/Public/Interfaces/ITargetPlatform.h b/Engine/Source/Developer/TargetPlatform/Public/Interfaces/ITargetPlatform.h
--- a/Engine/Source/Developer/TargetPlatform/Public/Interfaces/ITargetPlatform.h
+++ b/Engine/Source/Developer/TargetPlatform/Public/Interfaces/ITargetPlatform.h
@@ -160,6 +160,11 @@ public:
 	 */
 	virtual FString IniPlatformName( ) const = 0;
 
+	/**
+	 * Enables/Disable the device check
+	 */
+	virtual void EnableDeviceCheck(bool OnOff) = 0;
+
 	/**
 	 * Returns all discoverable physical devices.
 	 *
diff --git a/Engine/Source/Developer/Windows/WindowsTargetPlatform/Private/GenericWindowsTargetPlatform.h b/Engine/Source/Developer/Windows/WindowsTargetPlatform/Private/GenericWindowsTargetPlatform.h
--- a/Engine/Source/Developer/Windows/WindowsTargetPlatform/Private/GenericWindowsTargetPlatform.h
+++ b/Engine/Source/Developer/Windows/WindowsTargetPlatform/Private/GenericWindowsTargetPlatform.h
@@ -36,6 +36,8 @@ public:
 
 	// Begin ITargetPlatform interface
 
+	virtual void EnableDeviceCheck(bool OnOff) OVERRIDE {}
+
 	virtual void GetAllDevices( TArray<ITargetDevicePtr>& OutDevices ) const OVERRIDE
 	{
 		OutDevices.Reset();
diff --git a/Engine/Source/Developer/iOS/IOSTargetPlatform/Private/IOSDeviceHelper.h b/Engine/Source/Developer/iOS/IOSTargetPlatform/Private/IOSDeviceHelper.h
--- a/Engine/Source/Developer/iOS/IOSTargetPlatform/Private/IOSDeviceHelper.h
+++ b/Engine/Source/Developer/iOS/IOSTargetPlatform/Private/IOSDeviceHelper.h
@@ -54,6 +54,11 @@ public:
      */
     static bool InstallIPAOnDevice(const FTargetDeviceId& DeviceId, const FString& IPAPath);
     
+	/**
+	 * Suspends/Enables the device connect/disconnect thread
+	 */
+	static void EnableDeviceCheck(bool OnOff);
+
 private:
     static void DeviceCallback(void*);
     static void DoDeviceConnect(void*);
diff --git a/Engine/Source/Developer/iOS/IOSTargetPlatform/Private/IOSTargetPlatform.cpp b/Engine/Source/Developer/iOS/IOSTargetPlatform/Private/IOSTargetPlatform.cpp
--- a/Engine/Source/Developer/iOS/IOSTargetPlatform/Private/IOSTargetPlatform.cpp
+++ b/Engine/Source/Developer/iOS/IOSTargetPlatform/Private/IOSTargetPlatform.cpp
@@ -38,6 +38,11 @@ FIOSTargetPlatform::~FIOSTargetPlatform()
 /* ITargetPlatform interface
  *****************************************************************************/
 
+void FIOSTargetPlatform::EnableDeviceCheck(bool OnOff)
+{
+	FIOSDeviceHelper::EnableDeviceCheck(OnOff);
+}
+
 void FIOSTargetPlatform::GetAllDevices( TArray<ITargetDevicePtr>& OutDevices ) const
 {
 	OutDevices.Reset();
diff --git a/Engine/Source/Developer/iOS/IOSTargetPlatform/Private/IOSTargetPlatform.h b/Engine/Source/Developer/iOS/IOSTargetPlatform/Private/IOSTargetPlatform.h
--- a/Engine/Source/Developer/iOS/IOSTargetPlatform/Private/IOSTargetPlatform.h
+++ b/Engine/Source/Developer/iOS/IOSTargetPlatform/Private/IOSTargetPlatform.h
@@ -42,6 +42,8 @@ public:
 
 	// Begin ITargetPlatform interface
 
+	virtual void EnableDeviceCheck(bool OnOff) OVERRIDE;
+
 	virtual void GetAllDevices( TArray<ITargetDevicePtr>& OutDevices ) const OVERRIDE;
 
 	virtual ECompressionFlags GetBaseCompressionMethod() const OVERRIDE
diff --git a/Engine/Source/Developer/iOS/IOSTargetPlatform/Private/Mac/IOSDeviceHelperMac.cpp b/Engine/Source/Developer/iOS/IOSTargetPlatform/Private/Mac/IOSDeviceHelperMac.cpp
--- a/Engine/Source/Developer/iOS/IOSTargetPlatform/Private/Mac/IOSDeviceHelperMac.cpp
+++ b/Engine/Source/Developer/iOS/IOSTargetPlatform/Private/Mac/IOSDeviceHelperMac.cpp
@@ -387,3 +387,8 @@ bool FIOSDeviceHelper::InstallIPAOnDevice(const FTargetDeviceId& DeviceId, const
     }
     return false;
 }
+
+void FIOSDeviceHelper::EnableDeviceCheck(bool OnOff)
+{
+
+}
\ No newline at end of file
diff --git a/Engine/Source/Developer/iOS/IOSTargetPlatform/Private/Windows/IOSDeviceHelperWindows.cpp b/Engine/Source/Developer/iOS/IOSTargetPlatform/Private/Windows/IOSDeviceHelperWindows.cpp
--- a/Engine/Source/Developer/iOS/IOSTargetPlatform/Private/Windows/IOSDeviceHelperWindows.cpp
+++ b/Engine/Source/Developer/iOS/IOSTargetPlatform/Private/Windows/IOSDeviceHelperWindows.cpp
@@ -41,6 +41,7 @@ class FDeviceQueryTask
 public:
 	FDeviceQueryTask()
 		: Stopping(false)
+		, bCheckDevices(true)
 	{}
 
 	virtual bool Init() OVERRIDE
@@ -52,7 +53,10 @@ public:
 	{
 		while (!Stopping)
 		{
-			QueryDevices();
+			if (bCheckDevices)
+			{
+				QueryDevices();
+			}
 
 			FPlatformProcess::Sleep(5.0f);
 		}
@@ -73,6 +77,11 @@ public:
 		return DeviceNotification;
 	}
 
+	void Enable(bool OnOff)
+	{
+		bCheckDevices = OnOff;
+	}
+
 private:
 	bool ExecuteIPPCommand(const FString& CommandLine, FString* OutStdOut, FString* OutStdErr) const
 	{
@@ -201,6 +210,7 @@ private:
 	}
 
 	bool Stopping;
+	bool bCheckDevices;
 	TArray<FString> ConnectedDeviceIds;
 	FDeviceNotification DeviceNotification;
 };
@@ -319,3 +329,8 @@ bool FIOSDeviceHelper::InstallIPAOnDevice(const FTargetDeviceId& DeviceId, const
 {
     return false;
 }
+
+void FIOSDeviceHelper::EnableDeviceCheck(bool OnOff)
+{
+	QueryTask->Enable(OnOff);
+}
\ No newline at end of file
