#ue4
#ufe
* fix for server platform type being passed to UAT

diff --git a/Engine/Source/Developer/LauncherServices/Private/Launcher/LauncherBuildCommands.h b/Engine/Source/Developer/LauncherServices/Private/Launcher/LauncherBuildCommands.h
--- a/Engine/Source/Developer/LauncherServices/Private/Launcher/LauncherBuildCommands.h
+++ b/Engine/Source/Developer/LauncherServices/Private/Launcher/LauncherBuildCommands.h
@@ -76,6 +76,10 @@ public:
 		{
 			Platform = TEXT("Linux");
 		}
+		else if (TargetPlatform.PlatformName() == TEXT("WindowsServer") || TargetPlatform.PlatformName() == TEXT("WindowsNoEditor") || TargetPlatform.PlatformName() == TEXT("Windows"))
+		{
+			Platform = TEXT("Win64");
+		}
 		CommandLine += FString::Printf(TEXT(" -build -server -noclient -serverplatform=%s -skipcook"),
 			*Platform);
 
diff --git a/Engine/Source/Developer/LauncherServices/Private/Launcher/LauncherCookCommands.h b/Engine/Source/Developer/LauncherServices/Private/Launcher/LauncherCookCommands.h
--- a/Engine/Source/Developer/LauncherServices/Private/Launcher/LauncherCookCommands.h
+++ b/Engine/Source/Developer/LauncherServices/Private/Launcher/LauncherCookCommands.h
@@ -114,8 +114,17 @@ public:
 	{
 		FString CommandLine;
 
+		FString Platform = TEXT("Win64");
+		if (TargetPlatform.PlatformName() == TEXT("LinuxServer"))
+		{
+			Platform = TEXT("Linux");
+		}
+		else if (TargetPlatform.PlatformName() == TEXT("WindowsServer") || TargetPlatform.PlatformName() == TEXT("WindowsNoEditor") || TargetPlatform.PlatformName() == TEXT("Windows"))
+		{
+			Platform = TEXT("Win64");
+		}
 		CommandLine += FString::Printf(TEXT(" -cook -noclient -NoKill -server -serverplatform=%s"),
-			*TargetPlatform.PlatformName());
+			*Platform);
 
 		const TArray<FString>& CookedMaps = ChainState.Profile->GetCookedMaps();
 
diff --git a/Engine/Source/Developer/LauncherServices/Private/Launcher/LauncherDeployCommands.h b/Engine/Source/Developer/LauncherServices/Private/Launcher/LauncherDeployCommands.h
--- a/Engine/Source/Developer/LauncherServices/Private/Launcher/LauncherDeployCommands.h
+++ b/Engine/Source/Developer/LauncherServices/Private/Launcher/LauncherDeployCommands.h
@@ -159,6 +159,10 @@ public:
 		{
 			Platform = TEXT("Linux");
 		}
+		else if (TargetPlatform.PlatformName() == TEXT("WindowsServer") || TargetPlatform.PlatformName() == TEXT("WindowsNoEditor") || TargetPlatform.PlatformName() == TEXT("Windows"))
+		{
+			Platform = TEXT("Win64");
+		}
 		CommandLine = FString::Printf(TEXT(" -noclient -server -deploy -skipstage -serverplatform=%s -stagingdirectory=\"%s\" -cmdline=\"%s -InstanceName=\"Deployer (%s)\" -Messaging\""),
 			*Platform,
 			*StagePath,
@@ -343,6 +347,10 @@ public:
 		{
 			Platform = TEXT("Linux");
 		}
+		else if (TargetPlatform.PlatformName() == TEXT("WindowsServer") || TargetPlatform.PlatformName() == TEXT("WindowsNoEditor") || TargetPlatform.PlatformName() == TEXT("Windows"))
+		{
+			Platform = TEXT("Win64");
+		}
 		CommandLine = FString::Printf(TEXT(" -noclient -server -deploy -skipstage -serverplatform=%s -stagingdirectory=\"%s\" -cmdline=\"%s -InstanceName=\"Deployer (%s)\" -Messaging\""),
 			*Platform,
 			*ChainState.Profile->GetPackageDirectory(),
diff --git a/Engine/Source/Developer/LauncherServices/Private/Launcher/LauncherPackageCommands.h b/Engine/Source/Developer/LauncherServices/Private/Launcher/LauncherPackageCommands.h
--- a/Engine/Source/Developer/LauncherServices/Private/Launcher/LauncherPackageCommands.h
+++ b/Engine/Source/Developer/LauncherServices/Private/Launcher/LauncherPackageCommands.h
@@ -121,6 +121,10 @@ public:
 		{
 			Platform = TEXT("Linux");
 		}
+		else if (TargetPlatform.PlatformName() == TEXT("WindowsServer") || TargetPlatform.PlatformName() == TEXT("WindowsNoEditor") || TargetPlatform.PlatformName() == TEXT("Windows"))
+		{
+			Platform = TEXT("Win64");
+		}
 		CommandLine = FString::Printf(TEXT(" -noclient -server -skipcook -stage -serverplatform=%s"),
 			*Platform);
 
@@ -295,6 +299,10 @@ public:
 		{
 			Platform = TEXT("Linux");
 		}
+		else if (TargetPlatform.PlatformName() == TEXT("WindowsServer") || TargetPlatform.PlatformName() == TEXT("WindowsNoEditor") || TargetPlatform.PlatformName() == TEXT("Windows"))
+		{
+			Platform = TEXT("Win64");
+		}
 		CommandLine = FString::Printf(TEXT(" -noclient -server -skipcook -stage -package -serverplatform=%s"),
 			*Platform);
 
