Fix for build section not building the project
TTP326921
#ue4
#ufe

diff --git a/Engine/Source/Developer/LauncherServices/Private/Launcher/LauncherWorker.cpp b/Engine/Source/Developer/LauncherServices/Private/Launcher/LauncherWorker.cpp
--- a/Engine/Source/Developer/LauncherServices/Private/Launcher/LauncherWorker.cpp
+++ b/Engine/Source/Developer/LauncherServices/Private/Launcher/LauncherWorker.cpp
@@ -184,7 +184,7 @@ void FLauncherWorker::CreateAndExecuteTasks( const ILauncherProfileRef& InProfil
 	// create task chains
 	TaskChain = MakeShareable(new FLauncherVerifyProfileTask());
 	TArray<FString> Platforms;
-	if (InProfile->GetCookMode() == ELauncherProfileCookModes::ByTheBook)
+	if (InProfile->GetCookMode() == ELauncherProfileCookModes::ByTheBook || InProfile->IsBuilding())
 	{
 		Platforms = InProfile->GetCookedPlatforms();
 	}
