More changes to accomodate new Linux platforms.

Old "Linux" target platform is now called "LinuxNoEditor", so this needs to be patched up in places where target platform name was assumed to match platform name.

diff --git a/Engine/Source/Developer/LauncherServices/Private/Launcher/LauncherBuildCommands.h b/Engine/Source/Developer/LauncherServices/Private/Launcher/LauncherBuildCommands.h
--- a/Engine/Source/Developer/LauncherServices/Private/Launcher/LauncherBuildCommands.h
+++ b/Engine/Source/Developer/LauncherServices/Private/Launcher/LauncherBuildCommands.h
@@ -72,7 +72,7 @@ public:
 		FString CommandLine;
 
 		FString Platform = TEXT("Win64");
-		if (TargetPlatform.PlatformName() == TEXT("LinuxServer"))
+		if (TargetPlatform.PlatformName() == TEXT("LinuxServer") || TargetPlatform.PlatformName() == TEXT("LinuxNoEditor"))
 		{
 			Platform = TEXT("Linux");
 		}
diff --git a/Engine/Source/Developer/LauncherServices/Private/Launcher/LauncherCookCommands.h b/Engine/Source/Developer/LauncherServices/Private/Launcher/LauncherCookCommands.h
--- a/Engine/Source/Developer/LauncherServices/Private/Launcher/LauncherCookCommands.h
+++ b/Engine/Source/Developer/LauncherServices/Private/Launcher/LauncherCookCommands.h
@@ -115,7 +115,7 @@ public:
 		FString CommandLine;
 
 		FString Platform = TEXT("Win64");
-		if (TargetPlatform.PlatformName() == TEXT("LinuxServer"))
+		if (TargetPlatform.PlatformName() == TEXT("LinuxServer") || TargetPlatform.PlatformName() == TEXT("LinuxNoEditor"))
 		{
 			Platform = TEXT("Linux");
 		}
diff --git a/Engine/Source/Developer/LauncherServices/Private/Launcher/LauncherDeployCommands.h b/Engine/Source/Developer/LauncherServices/Private/Launcher/LauncherDeployCommands.h
--- a/Engine/Source/Developer/LauncherServices/Private/Launcher/LauncherDeployCommands.h
+++ b/Engine/Source/Developer/LauncherServices/Private/Launcher/LauncherDeployCommands.h
@@ -155,7 +155,7 @@ public:
 		}
 
 		FString Platform = TEXT("Win64");
-		if (TargetPlatform.PlatformName() == TEXT("LinuxServer"))
+		if (TargetPlatform.PlatformName() == TEXT("LinuxServer") || TargetPlatform.PlatformName() == TEXT("LinuxNoEditor"))
 		{
 			Platform = TEXT("Linux");
 		}
@@ -343,7 +343,7 @@ public:
 		}
 
 		FString Platform = TEXT("Win64");
-		if (TargetPlatform.PlatformName() == TEXT("LinuxServer"))
+		if (TargetPlatform.PlatformName() == TEXT("LinuxServer") || TargetPlatform.PlatformName() == TEXT("LinuxNoEditor"))
 		{
 			Platform = TEXT("Linux");
 		}
diff --git a/Engine/Source/Developer/LauncherServices/Private/Launcher/LauncherPackageCommands.h b/Engine/Source/Developer/LauncherServices/Private/Launcher/LauncherPackageCommands.h
--- a/Engine/Source/Developer/LauncherServices/Private/Launcher/LauncherPackageCommands.h
+++ b/Engine/Source/Developer/LauncherServices/Private/Launcher/LauncherPackageCommands.h
@@ -117,7 +117,7 @@ public:
 		}
 
 		FString Platform = TEXT("Win64");
-		if (TargetPlatform.PlatformName() == TEXT("LinuxServer"))
+		if (TargetPlatform.PlatformName() == TEXT("LinuxServer") || TargetPlatform.PlatformName() == TEXT("LinuxNoEditor"))
 		{
 			Platform = TEXT("Linux");
 		}
@@ -295,7 +295,7 @@ public:
 		}
 
 		FString Platform = TEXT("Win64");
-		if (TargetPlatform.PlatformName() == TEXT("LinuxServer"))
+		if (TargetPlatform.PlatformName() == TEXT("LinuxServer") || TargetPlatform.PlatformName() == TEXT("LinuxNoEditor"))
 		{
 			Platform = TEXT("Linux");
 		}
diff --git a/Engine/Source/Developer/LauncherServices/Private/Launcher/LauncherUATTask.h b/Engine/Source/Developer/LauncherServices/Private/Launcher/LauncherUATTask.h
--- a/Engine/Source/Developer/LauncherServices/Private/Launcher/LauncherUATTask.h
+++ b/Engine/Source/Developer/LauncherServices/Private/Launcher/LauncherUATTask.h
@@ -55,7 +55,7 @@ protected:
 
 		// switch server and no editor platforms to the proper type
 		FString Platform = TargetPlatform.PlatformName();
-		if (TargetPlatform.PlatformName() == TEXT("LinuxServer"))
+		if (TargetPlatform.PlatformName() == TEXT("LinuxServer") || TargetPlatform.PlatformName() == TEXT("LinuxNoEditor"))
 		{
 			Platform = TEXT("Linux");
 		}
diff --git a/Engine/Source/Editor/MainFrame/Private/Frame/MainFrameActions.cpp b/Engine/Source/Editor/MainFrame/Private/Frame/MainFrameActions.cpp
--- a/Engine/Source/Editor/MainFrame/Private/Frame/MainFrameActions.cpp
+++ b/Engine/Source/Editor/MainFrame/Private/Frame/MainFrameActions.cpp
@@ -553,6 +553,10 @@ void FMainFrameActionCallbacks::PackageProject( const FString InPlatformName, co
 	{
 		OptionalParams += TEXT(" -targetplatform=Mac");
 	}
+	else if (PlatformName == "LinuxNoEditor")
+	{
+		OptionalParams += TEXT(" -targetplatform=Linux");
+	}
 
 	// only build if the project has code that might need to be built
 	if (bProjectHasCode && FSourceCodeNavigation::IsCompilerAvailable())
