fix for game not waiting for Cook on the Fly server to initialize when running from Quick Launch
TTP333309
#ue4

diff --git a/Engine/Source/Developer/LauncherServices/Private/Launcher/LauncherCookCommands.h b/Engine/Source/Developer/LauncherServices/Private/Launcher/LauncherCookCommands.h
--- a/Engine/Source/Developer/LauncherServices/Private/Launcher/LauncherCookCommands.h
+++ b/Engine/Source/Developer/LauncherServices/Private/Launcher/LauncherCookCommands.h
@@ -222,8 +222,8 @@ public:
 			CommandLine += TEXT(" -iterate");
 		}
 
-		CommandLine += FString::Printf(TEXT(" -addcmdline=\"-InstanceName='File Server (Cook on the fly)' -SessionId=%s -SessionName='%s' -Messaging -timeout=%d%s\""),
-//			*InstanceId.ToString(),
+		CommandLine += FString::Printf(TEXT(" -addcmdline=\"-InstanceName='File Server (Cook on the fly)' -InstanceId=%s -SessionId=%s -SessionName='%s' -Messaging -timeout=%d%s\""),
+			*InstanceId.ToString(),
 			*ChainState.SessionId.ToString(),
 			*ChainState.Profile->GetName(),
             ChainState.Profile->GetTimeout(),
@@ -239,14 +239,14 @@ public:
 		CommandLine += TEXT(" -cookonthefly -skipserver");
 
 		// extract the port from the address list string
-/*		FString Address, Rest;
+		FString Address, Rest;
 		if(ChainState.FileServerAddressListString.Split(TEXT(":"), &Address, &Rest))
 		{
 			if (Rest.Split(TEXT("+"), &Address, &Rest))
 			{
 				CommandLine += FString::Printf(TEXT(" -port=%s"), *Address);
 			}
-		}*/
+		}
 		return CommandLine;
 	}
 
@@ -265,7 +265,7 @@ public:
 
 	virtual bool IsComplete() const OVERRIDE
 	{
-		return true;
+		return StartDetected;
 	}
 
 	virtual bool PostExecute(FLauncherTaskChainState& ChainState) OVERRIDE
