Docs: Removed file comments and added missing code documentation

diff --git a/Engine/Source/Developer/DeviceManager/Private/Models/DeviceDetailsCommands.h b/Engine/Source/Developer/DeviceManager/Private/Models/DeviceDetailsCommands.h
--- a/Engine/Source/Developer/DeviceManager/Private/Models/DeviceDetailsCommands.h
+++ b/Engine/Source/Developer/DeviceManager/Private/Models/DeviceDetailsCommands.h
@@ -1,9 +1,5 @@
 // Copyright 1998-2014 Epic Games, Inc. All Rights Reserved.
 
-/*=============================================================================
-	DeviceDetailsCommands.h: Declares the FDeviceDetailsCommands class.
-=============================================================================*/
-
 #pragma once
 
 
@@ -28,9 +24,9 @@ public:
 
 public:
 
-	// Begin TCommands interface
+	// TCommands interface
 
-	virtual void RegisterCommands( ) OVERRIDE
+	virtual void RegisterCommands( ) override
 	{
 		UI_COMMAND(Claim, "Claim", "Claim the device", EUserInterfaceActionType::Button, FInputGesture());
 		UI_COMMAND(Release, "Release", "Release the device", EUserInterfaceActionType::Button, FInputGesture());
@@ -46,8 +42,6 @@ public:
 		UI_COMMAND(Reboot, "Reboot", "Reboot the device", EUserInterfaceActionType::Button, FInputGesture());
 	}
 
-	// End TCommands interface
-
 public:
 
 	// ownership commands
diff --git a/Engine/Source/Developer/DeviceManager/Private/Models/DeviceDetailsFeature.h b/Engine/Source/Developer/DeviceManager/Private/Models/DeviceDetailsFeature.h
--- a/Engine/Source/Developer/DeviceManager/Private/Models/DeviceDetailsFeature.h
+++ b/Engine/Source/Developer/DeviceManager/Private/Models/DeviceDetailsFeature.h
@@ -1,9 +1,5 @@
 // Copyright 1998-2014 Epic Games, Inc. All Rights Reserved.
 
-/*=============================================================================
-	FDeviceDetailsFeature.h: Declares the FDeviceDetailsFeature class.
-=============================================================================*/
-
 #pragma once
 
 
diff --git a/Engine/Source/Developer/DeviceManager/Private/Models/DeviceManagerModel.h b/Engine/Source/Developer/DeviceManager/Private/Models/DeviceManagerModel.h
--- a/Engine/Source/Developer/DeviceManager/Private/Models/DeviceManagerModel.h
+++ b/Engine/Source/Developer/DeviceManager/Private/Models/DeviceManagerModel.h
@@ -1,9 +1,5 @@
 // Copyright 1998-2014 Epic Games, Inc. All Rights Reserved.
 
-/*=============================================================================
-	DeviceManagerModel.h: Declares the FDeviceManagerModel class.
-=============================================================================*/
-
 #pragma once
 
 
diff --git a/Engine/Source/Developer/DeviceManager/Private/Widgets/Apps/SDeviceApps.cpp b/Engine/Source/Developer/DeviceManager/Private/Widgets/Apps/SDeviceApps.cpp
--- a/Engine/Source/Developer/DeviceManager/Private/Widgets/Apps/SDeviceApps.cpp
+++ b/Engine/Source/Developer/DeviceManager/Private/Widgets/Apps/SDeviceApps.cpp
@@ -1,9 +1,5 @@
 // Copyright 1998-2014 Epic Games, Inc. All Rights Reserved.
 
-/*=============================================================================
-	SDeviceApps.cpp: Implements the SDeviceApps class.
-=============================================================================*/
-
 #include "DeviceManagerPrivatePCH.h"
 
 
diff --git a/Engine/Source/Developer/DeviceManager/Private/Widgets/Apps/SDeviceApps.h b/Engine/Source/Developer/DeviceManager/Private/Widgets/Apps/SDeviceApps.h
--- a/Engine/Source/Developer/DeviceManager/Private/Widgets/Apps/SDeviceApps.h
+++ b/Engine/Source/Developer/DeviceManager/Private/Widgets/Apps/SDeviceApps.h
@@ -1,9 +1,5 @@
 // Copyright 1998-2014 Epic Games, Inc. All Rights Reserved.
 
-/*=============================================================================
-	SDeviceApps.h: Declares the SDeviceApps class.
-=============================================================================*/
-
 #pragma once
 
 
diff --git a/Engine/Source/Developer/DeviceManager/Private/Widgets/Apps/SDeviceAppsAppListRow.h b/Engine/Source/Developer/DeviceManager/Private/Widgets/Apps/SDeviceAppsAppListRow.h
--- a/Engine/Source/Developer/DeviceManager/Private/Widgets/Apps/SDeviceAppsAppListRow.h
+++ b/Engine/Source/Developer/DeviceManager/Private/Widgets/Apps/SDeviceAppsAppListRow.h
@@ -1,9 +1,5 @@
 // Copyright 1998-2014 Epic Games, Inc. All Rights Reserved.
 
-/*=============================================================================
-	SDeviceAppsAppListRow.h: Declares the SDeviceAppsAppListRow class.
-=============================================================================*/
-
 #pragma once
 
 
diff --git a/Engine/Source/Developer/DeviceManager/Private/Widgets/Browser/SDeviceBrowser.cpp b/Engine/Source/Developer/DeviceManager/Private/Widgets/Browser/SDeviceBrowser.cpp
--- a/Engine/Source/Developer/DeviceManager/Private/Widgets/Browser/SDeviceBrowser.cpp
+++ b/Engine/Source/Developer/DeviceManager/Private/Widgets/Browser/SDeviceBrowser.cpp
@@ -1,9 +1,5 @@
 // Copyright 1998-2014 Epic Games, Inc. All Rights Reserved.
 
-/*=============================================================================
-	SDeviceBrowser.cpp: Implements the SDeviceBrowser class.
-=============================================================================*/
-
 #include "DeviceManagerPrivatePCH.h"
 
 
diff --git a/Engine/Source/Developer/DeviceManager/Private/Widgets/Browser/SDeviceBrowser.h b/Engine/Source/Developer/DeviceManager/Private/Widgets/Browser/SDeviceBrowser.h
--- a/Engine/Source/Developer/DeviceManager/Private/Widgets/Browser/SDeviceBrowser.h
+++ b/Engine/Source/Developer/DeviceManager/Private/Widgets/Browser/SDeviceBrowser.h
@@ -1,9 +1,5 @@
 // Copyright 1998-2014 Epic Games, Inc. All Rights Reserved.
 
-/*=============================================================================
-	SDeviceBrowser.h: Declares the SDeviceBrowser class.
-=============================================================================*/
-
 #pragma once
 
 
@@ -40,11 +36,9 @@ public:
 
 public:
 
-	// Begin SCompoundWidget overrides
-
-	virtual void Tick( const FGeometry& AllottedGeometry, const double InCurrentTime, const float InDeltaTime ) OVERRIDE;
+	// SCompoundWidget overrides
 
-	// End SCompoundWidget overrides
+	virtual void Tick( const FGeometry& AllottedGeometry, const double InCurrentTime, const float InDeltaTime ) override;
 
 protected:
 
diff --git a/Engine/Source/Developer/DeviceManager/Private/Widgets/Browser/SDeviceBrowserContextMenu.h b/Engine/Source/Developer/DeviceManager/Private/Widgets/Browser/SDeviceBrowserContextMenu.h
--- a/Engine/Source/Developer/DeviceManager/Private/Widgets/Browser/SDeviceBrowserContextMenu.h
+++ b/Engine/Source/Developer/DeviceManager/Private/Widgets/Browser/SDeviceBrowserContextMenu.h
@@ -1,9 +1,5 @@
 // Copyright 1998-2014 Epic Games, Inc. All Rights Reserved.
 
-/*=============================================================================
-	SDeviceBrowserContextMenu.h: Declares the SDeviceBrowserContextMenu class.
-=============================================================================*/
-
 #pragma once
 
 
diff --git a/Engine/Source/Developer/DeviceManager/Private/Widgets/Browser/SDeviceBrowserDeviceAdder.cpp b/Engine/Source/Developer/DeviceManager/Private/Widgets/Browser/SDeviceBrowserDeviceAdder.cpp
--- a/Engine/Source/Developer/DeviceManager/Private/Widgets/Browser/SDeviceBrowserDeviceAdder.cpp
+++ b/Engine/Source/Developer/DeviceManager/Private/Widgets/Browser/SDeviceBrowserDeviceAdder.cpp
@@ -1,9 +1,5 @@
 // Copyright 1998-2014 Epic Games, Inc. All Rights Reserved.
 
-/*=============================================================================
-	SDeviceBrowserDeviceAdder.cpp: Implements the SDeviceBrowserDeviceAdder class.
-=============================================================================*/
-
 #include "DeviceManagerPrivatePCH.h"
 
 
diff --git a/Engine/Source/Developer/DeviceManager/Private/Widgets/Browser/SDeviceBrowserDeviceAdder.h b/Engine/Source/Developer/DeviceManager/Private/Widgets/Browser/SDeviceBrowserDeviceAdder.h
--- a/Engine/Source/Developer/DeviceManager/Private/Widgets/Browser/SDeviceBrowserDeviceAdder.h
+++ b/Engine/Source/Developer/DeviceManager/Private/Widgets/Browser/SDeviceBrowserDeviceAdder.h
@@ -1,9 +1,5 @@
 // Copyright 1998-2014 Epic Games, Inc. All Rights Reserved.
 
-/*=============================================================================
-	SDeviceBrowserDeviceAdder.h: Declares the SDeviceBrowserDeviceAdder class.
-=============================================================================*/
-
 #pragma once
 
 
diff --git a/Engine/Source/Developer/DeviceManager/Private/Widgets/Browser/SDeviceBrowserDeviceListRow.h b/Engine/Source/Developer/DeviceManager/Private/Widgets/Browser/SDeviceBrowserDeviceListRow.h
--- a/Engine/Source/Developer/DeviceManager/Private/Widgets/Browser/SDeviceBrowserDeviceListRow.h
+++ b/Engine/Source/Developer/DeviceManager/Private/Widgets/Browser/SDeviceBrowserDeviceListRow.h
@@ -1,9 +1,5 @@
 // Copyright 1998-2014 Epic Games, Inc. All Rights Reserved.
 
-/*=============================================================================
-	SDeviceBrowserDeviceListRow.h: Declares the SDeviceBrowserDeviceListRow class.
-=============================================================================*/
-
 #pragma once
 
 
diff --git a/Engine/Source/Developer/DeviceManager/Private/Widgets/Browser/SDeviceBrowserFilterBar.cpp b/Engine/Source/Developer/DeviceManager/Private/Widgets/Browser/SDeviceBrowserFilterBar.cpp
--- a/Engine/Source/Developer/DeviceManager/Private/Widgets/Browser/SDeviceBrowserFilterBar.cpp
+++ b/Engine/Source/Developer/DeviceManager/Private/Widgets/Browser/SDeviceBrowserFilterBar.cpp
@@ -1,9 +1,5 @@
 // Copyright 1998-2014 Epic Games, Inc. All Rights Reserved.
 
-/*=============================================================================
-	SDeviceBrowserFilterBar.cpp: Implements the SDeviceBrowserFilterBar class.
-=============================================================================*/
-
 #include "DeviceManagerPrivatePCH.h"
 
 
diff --git a/Engine/Source/Developer/DeviceManager/Private/Widgets/Browser/SDeviceBrowserFilterBar.h b/Engine/Source/Developer/DeviceManager/Private/Widgets/Browser/SDeviceBrowserFilterBar.h
--- a/Engine/Source/Developer/DeviceManager/Private/Widgets/Browser/SDeviceBrowserFilterBar.h
+++ b/Engine/Source/Developer/DeviceManager/Private/Widgets/Browser/SDeviceBrowserFilterBar.h
@@ -1,9 +1,5 @@
 // Copyright 1998-2014 Epic Games, Inc. All Rights Reserved.
 
-/*=============================================================================
-	SDeviceBrowserFilterBar.h: Declares the SDeviceBrowserFilterBar class.
-=============================================================================*/
-
 #pragma once
 
 
diff --git a/Engine/Source/Developer/DeviceManager/Private/Widgets/Browser/SDeviceBrowserTooltip.h b/Engine/Source/Developer/DeviceManager/Private/Widgets/Browser/SDeviceBrowserTooltip.h
--- a/Engine/Source/Developer/DeviceManager/Private/Widgets/Browser/SDeviceBrowserTooltip.h
+++ b/Engine/Source/Developer/DeviceManager/Private/Widgets/Browser/SDeviceBrowserTooltip.h
@@ -1,9 +1,5 @@
 // Copyright 1998-2014 Epic Games, Inc. All Rights Reserved.
 
-/*=============================================================================
-	SDeviceBrowserTooltip.h: Declares the SDeviceBrowserTooltip class.
-=============================================================================*/
-
 #pragma once
 
 
diff --git a/Engine/Source/Developer/DeviceManager/Private/Widgets/Details/SDeviceDetails.cpp b/Engine/Source/Developer/DeviceManager/Private/Widgets/Details/SDeviceDetails.cpp
--- a/Engine/Source/Developer/DeviceManager/Private/Widgets/Details/SDeviceDetails.cpp
+++ b/Engine/Source/Developer/DeviceManager/Private/Widgets/Details/SDeviceDetails.cpp
@@ -1,9 +1,5 @@
 // Copyright 1998-2014 Epic Games, Inc. All Rights Reserved.
 
-/*=============================================================================
-	SDeviceDetails.cpp: Implements the SDeviceDetails class.
-=============================================================================*/
-
 #include "DeviceManagerPrivatePCH.h"
 
 
diff --git a/Engine/Source/Developer/DeviceManager/Private/Widgets/Details/SDeviceDetails.h b/Engine/Source/Developer/DeviceManager/Private/Widgets/Details/SDeviceDetails.h
--- a/Engine/Source/Developer/DeviceManager/Private/Widgets/Details/SDeviceDetails.h
+++ b/Engine/Source/Developer/DeviceManager/Private/Widgets/Details/SDeviceDetails.h
@@ -1,9 +1,5 @@
 // Copyright 1998-2014 Epic Games, Inc. All Rights Reserved.
 
-/*=============================================================================
-	SDeviceDetails.h: Declares the SDeviceDetails class.
-=============================================================================*/
-
 #pragma once
 
 
diff --git a/Engine/Source/Developer/DeviceManager/Private/Widgets/Details/SDeviceDetailsFeatureListRow.h b/Engine/Source/Developer/DeviceManager/Private/Widgets/Details/SDeviceDetailsFeatureListRow.h
--- a/Engine/Source/Developer/DeviceManager/Private/Widgets/Details/SDeviceDetailsFeatureListRow.h
+++ b/Engine/Source/Developer/DeviceManager/Private/Widgets/Details/SDeviceDetailsFeatureListRow.h
@@ -1,9 +1,5 @@
 // Copyright 1998-2014 Epic Games, Inc. All Rights Reserved.
 
-/*=============================================================================
-	SDeviceDetailsFeatureListRow.h: Declares the SDeviceDetailsFeatureListRow class.
-=============================================================================*/
-
 #pragma once
 
 
@@ -42,10 +38,10 @@ public:
 
 public:
 
-	// Begin SMultiColumnTableRow interface
+	// SMultiColumnTableRow interface
 
 	BEGIN_SLATE_FUNCTION_BUILD_OPTIMIZATION
-	virtual TSharedRef<SWidget> GenerateWidgetForColumn( const FName& ColumnName ) OVERRIDE
+	virtual TSharedRef<SWidget> GenerateWidgetForColumn( const FName& ColumnName ) override
 	{
 		if (ColumnName == "Available")
 		{
@@ -74,8 +70,6 @@ public:
 	}
 	END_SLATE_FUNCTION_BUILD_OPTIMIZATION
 
-	// End SMultiColumnTableRow interface
-
 private:
 
 	// Callback for getting the text color.
diff --git a/Engine/Source/Developer/DeviceManager/Private/Widgets/Processes/SDeviceProcesses.cpp b/Engine/Source/Developer/DeviceManager/Private/Widgets/Processes/SDeviceProcesses.cpp
--- a/Engine/Source/Developer/DeviceManager/Private/Widgets/Processes/SDeviceProcesses.cpp
+++ b/Engine/Source/Developer/DeviceManager/Private/Widgets/Processes/SDeviceProcesses.cpp
@@ -1,9 +1,5 @@
 // Copyright 1998-2014 Epic Games, Inc. All Rights Reserved.
 
-/*=============================================================================
-	SDeviceProcesses.cpp: Implements the SDeviceProcesses class.
-=============================================================================*/
-
 #include "DeviceManagerPrivatePCH.h"
 
 
diff --git a/Engine/Source/Developer/DeviceManager/Private/Widgets/Processes/SDeviceProcesses.h b/Engine/Source/Developer/DeviceManager/Private/Widgets/Processes/SDeviceProcesses.h
--- a/Engine/Source/Developer/DeviceManager/Private/Widgets/Processes/SDeviceProcesses.h
+++ b/Engine/Source/Developer/DeviceManager/Private/Widgets/Processes/SDeviceProcesses.h
@@ -1,9 +1,5 @@
 // Copyright 1998-2014 Epic Games, Inc. All Rights Reserved.
 
-/*=============================================================================
-	SDeviceProcesses.h: Declares the SDeviceProcesses class.
-=============================================================================*/
-
 #pragma once
 
 
@@ -37,11 +33,9 @@ public:
 
 public:
 
-	// Begin SCompoundWidget overrides
-
-	void Tick( const FGeometry& AllottedGeometry, const double InCurrentTime, const float InDeltaTime );
+	// SCompoundWidget overrides
 
-	// End SCompoundWidget overrides
+	virtual void Tick( const FGeometry& AllottedGeometry, const double InCurrentTime, const float InDeltaTime ) override;
 
 protected:
 
diff --git a/Engine/Source/Developer/DeviceManager/Private/Widgets/Processes/SDeviceProcessesProcessListRow.h b/Engine/Source/Developer/DeviceManager/Private/Widgets/Processes/SDeviceProcessesProcessListRow.h
--- a/Engine/Source/Developer/DeviceManager/Private/Widgets/Processes/SDeviceProcessesProcessListRow.h
+++ b/Engine/Source/Developer/DeviceManager/Private/Widgets/Processes/SDeviceProcessesProcessListRow.h
@@ -1,9 +1,5 @@
 // Copyright 1998-2014 Epic Games, Inc. All Rights Reserved.
 
-/*=============================================================================
-	SDeviceProcessesProcessListRow.h: Declares the SDeviceProcessesProcessListRow class.
-=============================================================================*/
-
 #pragma once
 
 
diff --git a/Engine/Source/Developer/DeviceManager/Private/Widgets/Processes/SDeviceProcessesProcessTreeNode.h b/Engine/Source/Developer/DeviceManager/Private/Widgets/Processes/SDeviceProcessesProcessTreeNode.h
--- a/Engine/Source/Developer/DeviceManager/Private/Widgets/Processes/SDeviceProcessesProcessTreeNode.h
+++ b/Engine/Source/Developer/DeviceManager/Private/Widgets/Processes/SDeviceProcessesProcessTreeNode.h
@@ -1,9 +1,5 @@
 // Copyright 1998-2014 Epic Games, Inc. All Rights Reserved.
 
-/*=============================================================================
-	DeviceProcessesProcessTreeNode.h: Declares the FDeviceProcessesProcessTreeNode class.
-=============================================================================*/
-
 #pragma once
 
 
diff --git a/Engine/Source/Developer/DeviceManager/Private/Widgets/SDeviceManager.cpp b/Engine/Source/Developer/DeviceManager/Private/Widgets/SDeviceManager.cpp
--- a/Engine/Source/Developer/DeviceManager/Private/Widgets/SDeviceManager.cpp
+++ b/Engine/Source/Developer/DeviceManager/Private/Widgets/SDeviceManager.cpp
@@ -1,9 +1,5 @@
 // Copyright 1998-2014 Epic Games, Inc. All Rights Reserved.
 
-/*=============================================================================
-	SDeviceManager.cpp: Implements the SDeviceManager class.
-=============================================================================*/
-
 #include "DeviceManagerPrivatePCH.h"
 
 
diff --git a/Engine/Source/Developer/DeviceManager/Private/Widgets/SDeviceManager.h b/Engine/Source/Developer/DeviceManager/Private/Widgets/SDeviceManager.h
--- a/Engine/Source/Developer/DeviceManager/Private/Widgets/SDeviceManager.h
+++ b/Engine/Source/Developer/DeviceManager/Private/Widgets/SDeviceManager.h
@@ -1,9 +1,5 @@
 // Copyright 1998-2014 Epic Games, Inc. All Rights Reserved.
 
-/*=============================================================================
-	SDeviceManager.h: Declares the SDeviceManager class.
-=============================================================================*/
-
 #pragma once
 
 
diff --git a/Engine/Source/Developer/DeviceManager/Private/Widgets/Shared/SDeviceQuickInfo.h b/Engine/Source/Developer/DeviceManager/Private/Widgets/Shared/SDeviceQuickInfo.h
--- a/Engine/Source/Developer/DeviceManager/Private/Widgets/Shared/SDeviceQuickInfo.h
+++ b/Engine/Source/Developer/DeviceManager/Private/Widgets/Shared/SDeviceQuickInfo.h
@@ -1,9 +1,5 @@
 // Copyright 1998-2014 Epic Games, Inc. All Rights Reserved.
 
-/*=============================================================================
-	SDeviceQuickInfo.h: Declares the SDeviceQuickInfo class.
-=============================================================================*/
-
 #pragma once
 
 
diff --git a/Engine/Source/Developer/DeviceManager/Private/Widgets/Toolbar/SDeviceToolbar.cpp b/Engine/Source/Developer/DeviceManager/Private/Widgets/Toolbar/SDeviceToolbar.cpp
--- a/Engine/Source/Developer/DeviceManager/Private/Widgets/Toolbar/SDeviceToolbar.cpp
+++ b/Engine/Source/Developer/DeviceManager/Private/Widgets/Toolbar/SDeviceToolbar.cpp
@@ -1,9 +1,5 @@
 // Copyright 1998-2014 Epic Games, Inc. All Rights Reserved.
 
-/*=============================================================================
-	SDeviceToolbar.cpp: Implements the SDeviceToolbar class.
-=============================================================================*/
-
 #include "DeviceManagerPrivatePCH.h"
 
 
diff --git a/Engine/Source/Developer/DeviceManager/Private/Widgets/Toolbar/SDeviceToolbar.h b/Engine/Source/Developer/DeviceManager/Private/Widgets/Toolbar/SDeviceToolbar.h
--- a/Engine/Source/Developer/DeviceManager/Private/Widgets/Toolbar/SDeviceToolbar.h
+++ b/Engine/Source/Developer/DeviceManager/Private/Widgets/Toolbar/SDeviceToolbar.h
@@ -1,9 +1,5 @@
 // Copyright 1998-2014 Epic Games, Inc. All Rights Reserved.
 
-/*=============================================================================
-	SDeviceToolbar.h: Declares the SDeviceToolbar class.
-=============================================================================*/
-
 #pragma once
 
 
diff --git a/Engine/Source/Developer/DeviceManager/Public/DeviceManager.h b/Engine/Source/Developer/DeviceManager/Public/DeviceManager.h
--- a/Engine/Source/Developer/DeviceManager/Public/DeviceManager.h
+++ b/Engine/Source/Developer/DeviceManager/Public/DeviceManager.h
@@ -1,9 +1,5 @@
 // Copyright 1998-2014 Epic Games, Inc. All Rights Reserved.
 
-/*=============================================================================
-	DeviceManager.h: DeviceManager module public header file.
-=============================================================================*/
-
 #pragma once
 
 
diff --git a/Engine/Source/Developer/DeviceManager/Public/Interfaces/IDeviceManagerModule.h b/Engine/Source/Developer/DeviceManager/Public/Interfaces/IDeviceManagerModule.h
--- a/Engine/Source/Developer/DeviceManager/Public/Interfaces/IDeviceManagerModule.h
+++ b/Engine/Source/Developer/DeviceManager/Public/Interfaces/IDeviceManagerModule.h
@@ -1,9 +1,5 @@
 // Copyright 1998-2014 Epic Games, Inc. All Rights Reserved.
 
-/*=============================================================================
-	IDeviceManagerModule.h: Declares the IDeviceManagerModule interface.
-=============================================================================*/
-
 #pragma once
 
 
diff --git a/Engine/Source/Developer/GammaUI/Private/GammaUIPrivatePCH.h b/Engine/Source/Developer/GammaUI/Private/GammaUIPrivatePCH.h
--- a/Engine/Source/Developer/GammaUI/Private/GammaUIPrivatePCH.h
+++ b/Engine/Source/Developer/GammaUI/Private/GammaUIPrivatePCH.h
@@ -3,4 +3,4 @@
 #pragma once
 
 #include "Engine.h"
-#include "../Public/GammaUI.h"
+#include "GammaUI.h"
diff --git a/Engine/Source/Developer/ImageCore/Private/ImageCore.cpp b/Engine/Source/Developer/ImageCore/Private/ImageCore.cpp
--- a/Engine/Source/Developer/ImageCore/Private/ImageCore.cpp
+++ b/Engine/Source/Developer/ImageCore/Private/ImageCore.cpp
@@ -1,9 +1,5 @@
 // Copyright 1998-2014 Epic Games, Inc. All Rights Reserved.
 
-/*=============================================================================
-	ImageCore.cpp: Implements the ImageCore module.
-=============================================================================*/
-
 #include "ImageCorePCH.h"
 
 
@@ -25,6 +21,7 @@ static void InitImageStorage(FImage& Image)
 	Image.RawData.AddUninitialized(NumBytes);
 }
 
+
 /**
  * Copies an image accounting for format differences. Sizes must match.
  *
diff --git a/Engine/Source/Developer/ImageCore/Private/ImageCorePCH.h b/Engine/Source/Developer/ImageCore/Private/ImageCorePCH.h
--- a/Engine/Source/Developer/ImageCore/Private/ImageCorePCH.h
+++ b/Engine/Source/Developer/ImageCore/Private/ImageCorePCH.h
@@ -1,10 +1,6 @@
 // Copyright 1998-2014 Epic Games, Inc. All Rights Reserved.
 
-/*=============================================================================
-	ImgaeCorePrivatePCH.h: Pre-compiled header file for the ImageCore module.
-=============================================================================*/
-
-#include "../Public/ImageCore.h"
+#include "ImageCore.h"
 
 
 /* Dependencies
diff --git a/Engine/Source/Developer/ImageCore/Public/ImageCore.h b/Engine/Source/Developer/ImageCore/Public/ImageCore.h
--- a/Engine/Source/Developer/ImageCore/Public/ImageCore.h
+++ b/Engine/Source/Developer/ImageCore/Public/ImageCore.h
@@ -1,9 +1,5 @@
 // Copyright 1998-2014 Epic Games, Inc. All Rights Reserved.
 
-/*=============================================================================
-	ImageCore.h: ImageCore module public header file.
-=============================================================================*/
-
 #pragma once
 
 
diff --git a/Engine/Source/Developer/ImageWrapper/Private/BmpImageWrapper.h b/Engine/Source/Developer/ImageWrapper/Private/BmpImageWrapper.h
--- a/Engine/Source/Developer/ImageWrapper/Private/BmpImageWrapper.h
+++ b/Engine/Source/Developer/ImageWrapper/Private/BmpImageWrapper.h
@@ -17,17 +17,6 @@ public:
 
 public:
 
-	// Begin FImageWrapper Interface
-
-	virtual void Compress( int32 Quality ) OVERRIDE;
-
-	virtual void Uncompress( const ERGBFormat::Type InFormat, int32 InBitDepth ) OVERRIDE;
-	
-	virtual bool SetCompressed( const void* InCompressedData, int32 InCompressedSize ) OVERRIDE;
-
-	// End FImageWrapper Interface
-
-
 	/** Helper function used to uncompress BMP data from a buffer */
 	void UncompressBMPData( const ERGBFormat::Type InFormat, const int32 InBitDepth );
 
@@ -45,7 +34,16 @@ public:
 	 */
 	bool LoadBMPInfoHeader();
 
+public:
+
+	// FImageWrapper Interface
+
+	virtual void Compress( int32 Quality ) override;
+	virtual void Uncompress( const ERGBFormat::Type InFormat, int32 InBitDepth ) override;
+	virtual bool SetCompressed( const void* InCompressedData, int32 InCompressedSize ) override;
+
 private:
+
 	/** Whether this file has a BMP file header */
 	bool bHasHeader;
 
diff --git a/Engine/Source/Developer/ImageWrapper/Private/IcoImageWrapper.h b/Engine/Source/Developer/ImageWrapper/Private/IcoImageWrapper.h
--- a/Engine/Source/Developer/ImageWrapper/Private/IcoImageWrapper.h
+++ b/Engine/Source/Developer/ImageWrapper/Private/IcoImageWrapper.h
@@ -17,17 +17,12 @@ public:
 
 public:
 
-	// Begin FImageWrapper Interface
+	// FImageWrapper Interface
 
-	virtual void Compress( int32 Quality ) OVERRIDE;
-
-	virtual void Uncompress( const ERGBFormat::Type InFormat, int32 InBitDepth ) OVERRIDE;
-	
-	virtual bool SetCompressed( const void* InCompressedData, int32 InCompressedSize ) OVERRIDE;
-
-	virtual bool GetRaw( const ERGBFormat::Type InFormat, int32 InBitDepth, const TArray<uint8>*& OutRawData ) OVERRIDE;
-
-	// End FImageWrapper Interface
+	virtual void Compress( int32 Quality ) override;
+	virtual void Uncompress( const ERGBFormat::Type InFormat, int32 InBitDepth ) override;
+	virtual bool SetCompressed( const void* InCompressedData, int32 InCompressedSize ) override;
+	virtual bool GetRaw( const ERGBFormat::Type InFormat, int32 InBitDepth, const TArray<uint8>*& OutRawData ) override;
 
 private:
 
diff --git a/Engine/Source/Developer/ImageWrapper/Private/ImageWrapperBase.cpp b/Engine/Source/Developer/ImageWrapper/Private/ImageWrapperBase.cpp
--- a/Engine/Source/Developer/ImageWrapper/Private/ImageWrapperBase.cpp
+++ b/Engine/Source/Developer/ImageWrapper/Private/ImageWrapperBase.cpp
@@ -1,9 +1,5 @@
 // Copyright 1998-2014 Epic Games, Inc. All Rights Reserved.
 
-/*=============================================================================
-	ImageWrapperBase.cpp: Implements the FImageWrapperBase class.
-=============================================================================*/
-
 #include "ImageWrapperPrivatePCH.h"
 
 
@@ -83,6 +79,7 @@ bool FImageWrapperBase::SetCompressed( const void* InCompressedData, int32 InCom
 	return true;
 }
 
+
 bool FImageWrapperBase::SetRaw( const void* InRawData, int32 InRawSize, const int32 InWidth, const int32 InHeight, const ERGBFormat::Type InFormat, const int32 InBitDepth )
 {
 	check(InRawData != NULL);
diff --git a/Engine/Source/Developer/ImageWrapper/Private/ImageWrapperBase.h b/Engine/Source/Developer/ImageWrapper/Private/ImageWrapperBase.h
--- a/Engine/Source/Developer/ImageWrapper/Private/ImageWrapperBase.h
+++ b/Engine/Source/Developer/ImageWrapper/Private/ImageWrapperBase.h
@@ -1,9 +1,5 @@
 // Copyright 1998-2014 Epic Games, Inc. All Rights Reserved.
 
-/*=============================================================================
-	ImageWrapperBase.h: Declares the FImageWrapperBase class.
-=============================================================================*/
-
 #pragma once
 
 
@@ -20,11 +16,24 @@ public:
 	 */
 	FImageWrapperBase( );
 
+public:
+
+	/**
+	 * Gets the image's raw data.
+	 *
+	 * @return A read-only byte array containing the data.
+	 */
+	const TArray<uint8>& GetRawData() const
+	{
+		return RawData;
+	}
 
 public:
 
 	/**
 	 * Compresses the data.
+	 *
+	 * @param Quality The compression quality.
 	 */
 	virtual void Compress( int32 Quality ) = 0;
 
@@ -35,55 +44,48 @@ public:
 
 	/**
 	 * Sets last error message.
+	 *
+	 * @param ErrorMessage The error message to set.
 	 */
 	virtual void SetError( const TCHAR* ErrorMessage );
 
 	/**  
 	 * Function to uncompress our data 
 	 *
-	 * @param	InFormat		how we want to manipulate the RGB data
+	 * @param InFormat How we want to manipulate the RGB data
 	 */
 	virtual void Uncompress( const ERGBFormat::Type InFormat, int32 InBitDepth ) = 0;
 
-
 public:
 
-	// Begin IImageWrapper interface
+	// IImageWrapper interface
 
-	virtual const TArray<uint8>& GetCompressed( int32 Quality = 0 ) OVERRIDE;
+	virtual const TArray<uint8>& GetCompressed( int32 Quality = 0 ) override;
 
-	virtual int32 GetBitDepth( ) const OVERRIDE
+	virtual int32 GetBitDepth( ) const override
 	{
 		return BitDepth;
 	}
 
-	virtual ERGBFormat::Type GetFormat() const OVERRIDE
+	virtual ERGBFormat::Type GetFormat() const override
 	{
 		return Format;
 	}
 
-	virtual int32 GetHeight( ) const OVERRIDE
+	virtual int32 GetHeight( ) const override
 	{
 		return Height;
 	}
 
-	virtual bool GetRaw( const ERGBFormat::Type InFormat, int32 InBitDepth, const TArray<uint8>*& OutRawData ) OVERRIDE;
+	virtual bool GetRaw( const ERGBFormat::Type InFormat, int32 InBitDepth, const TArray<uint8>*& OutRawData ) override;
 
-	virtual int32 GetWidth( ) const OVERRIDE
+	virtual int32 GetWidth( ) const override
 	{
 		return Width;
 	}
 
-	virtual bool SetCompressed( const void* InCompressedData, int32 InCompressedSize ) OVERRIDE;
-
-	virtual bool SetRaw( const void* InRawData, int32 InRawSize, const int32 InWidth, const int32 InHeight, const ERGBFormat::Type InFormat, const int32 InBitDepth ) OVERRIDE;
-
-	// End IImageWrapper interface
-
-	const TArray<uint8>& GetRawData() const
-	{
-		return RawData;
-	}
+	virtual bool SetCompressed( const void* InCompressedData, int32 InCompressedSize ) override;
+	virtual bool SetRaw( const void* InRawData, int32 InRawSize, const int32 InWidth, const int32 InHeight, const ERGBFormat::Type InFormat, const int32 InBitDepth ) override;
 
 protected:
 
diff --git a/Engine/Source/Developer/ImageWrapper/Private/ImageWrapperModule.cpp b/Engine/Source/Developer/ImageWrapper/Private/ImageWrapperModule.cpp
--- a/Engine/Source/Developer/ImageWrapper/Private/ImageWrapperModule.cpp
+++ b/Engine/Source/Developer/ImageWrapper/Private/ImageWrapperModule.cpp
@@ -1,9 +1,5 @@
 // Copyright 1998-2014 Epic Games, Inc. All Rights Reserved.
 
-/*=============================================================================
-	ImageWrapperModule.cpp: Implements the FImageWrapperModule class.
-=============================================================================*/
-
 #include "ImageWrapperPrivatePCH.h"
 
 
@@ -18,9 +14,9 @@ class FImageWrapperModule
 {
 public:
 
-	// Begin IImageWrapperModule interface
+	// IImageWrapperModule interface
 
-	virtual IImageWrapperPtr CreateImageWrapper( const EImageFormat::Type InFormat )
+	virtual IImageWrapperPtr CreateImageWrapper( const EImageFormat::Type InFormat ) override
 	{
 		FImageWrapperBase* ImageWrapper = NULL;
 
@@ -53,18 +49,12 @@ public:
 		return MakeShareable(ImageWrapper);
 	}
 
-	// End IImageWrapperModule interface
-
-
 public:
 
-	// Begin IModuleInterface interface
-
-	virtual void StartupModule( ) OVERRIDE { }
-
-	virtual void ShutdownModule( ) OVERRIDE { }
+	// IModuleInterface interface
 
-	// End IModuleInterface interface
+	virtual void StartupModule( ) override { }
+	virtual void ShutdownModule( ) override { }
 };
 
 
diff --git a/Engine/Source/Developer/ImageWrapper/Private/ImageWrapperPrivatePCH.h b/Engine/Source/Developer/ImageWrapper/Private/ImageWrapperPrivatePCH.h
--- a/Engine/Source/Developer/ImageWrapper/Private/ImageWrapperPrivatePCH.h
+++ b/Engine/Source/Developer/ImageWrapper/Private/ImageWrapperPrivatePCH.h
@@ -1,13 +1,8 @@
 // Copyright 1998-2014 Epic Games, Inc. All Rights Reserved.
 
-/*=============================================================================
-	ImageWrapperPrivatePCH.h: Pre-compiled header file for the ImageWrapper module.
-=============================================================================*/
-
 #pragma once
 
-
-#include "../Public/ImageWrapper.h"
+#include "ImageWrapper.h"
 
 
 /* Dependencies
diff --git a/Engine/Source/Developer/ImageWrapper/Private/JpegImageWrapper.cpp b/Engine/Source/Developer/ImageWrapper/Private/JpegImageWrapper.cpp
--- a/Engine/Source/Developer/ImageWrapper/Private/JpegImageWrapper.cpp
+++ b/Engine/Source/Developer/ImageWrapper/Private/JpegImageWrapper.cpp
@@ -1,15 +1,8 @@
 // Copyright 1998-2014 Epic Games, Inc. All Rights Reserved.
 
-/*=============================================================================
-	JpegImageWrapper.h: Declares the FJpegImageWrapper class.
-
-	
-	Source code for JPEG decompression from:
-	http://code.google.com/p/jpeg-compressor/
-=============================================================================*/
-
 #include "ImageWrapperPrivatePCH.h"
 
+
 #if WITH_UNREALJPEG
 
 #ifdef __clang__
@@ -38,7 +31,7 @@ FCriticalSection GJPEGSection;
 
 FJpegImageWrapper::FJpegImageWrapper( )
 	: FImageWrapperBase()
-{}
+{ }
 
 
 /* FImageWrapperBase interface
@@ -75,6 +68,7 @@ bool FJpegImageWrapper::SetCompressed( const void* InCompressedData, int32 InCom
 	return bResult;
 }
 
+
 bool FJpegImageWrapper::SetRaw( const void* InRawData, int32 InRawSize, const int32 InWidth, const int32 InHeight, const ERGBFormat::Type InFormat, const int32 InBitDepth )
 {
 	check((InFormat == ERGBFormat::RGBA || InFormat == ERGBFormat::BGRA) && InBitDepth == 8);
@@ -84,6 +78,7 @@ bool FJpegImageWrapper::SetRaw( const void* InRawData, int32 InRawSize, const in
 	return bResult;
 }
 
+
 void FJpegImageWrapper::Compress( int32 Quality )
 {
 	if (CompressedData.Num() == 0)
@@ -127,6 +122,7 @@ void FJpegImageWrapper::Compress( int32 Quality )
 	}
 }
 
+
 void FJpegImageWrapper::Uncompress( const ERGBFormat::Type InFormat, int32 InBitDepth )
 {
 	// Ensure we haven't already uncompressed the file.
diff --git a/Engine/Source/Developer/ImageWrapper/Private/JpegImageWrapper.h b/Engine/Source/Developer/ImageWrapper/Private/JpegImageWrapper.h
--- a/Engine/Source/Developer/ImageWrapper/Private/JpegImageWrapper.h
+++ b/Engine/Source/Developer/ImageWrapper/Private/JpegImageWrapper.h
@@ -1,16 +1,14 @@
 // Copyright 1998-2014 Epic Games, Inc. All Rights Reserved.
 
-/*=============================================================================
-	JpegImageWrapper.h: Declares the FJpegImageWrapper class.
-=============================================================================*/
-
 #pragma once
 
 #if WITH_UNREALJPEG
 
 
 /**
- * Uncompresses JPEG data to raw 24bit RGB image that can be used by Unreal textures
+ * Uncompresses JPEG data to raw 24bit RGB image that can be used by Unreal textures.
+ *
+ * Source code for JPEG decompression from http://code.google.com/p/jpeg-compressor/
  */
 class FJpegImageWrapper
 	: public FImageWrapperBase
@@ -22,16 +20,14 @@ public:
 	 */
 	FJpegImageWrapper();
 
-
 public:
 
-	// Begin FImageWrapperBase interface
-	virtual bool SetCompressed( const void* InCompressedData, int32 InCompressedSize ) OVERRIDE;
-	virtual bool SetRaw( const void* InRawData, int32 InRawSize, const int32 InWidth, const int32 InHeight, const ERGBFormat::Type InFormat, const int32 InBitDepth ) OVERRIDE;
+	// FImageWrapperBase interface
 
-	virtual void Uncompress( const ERGBFormat::Type InFormat, int32 InBitDepth ) OVERRIDE;
-	virtual void Compress(int32 Quality) OVERRIDE;
-	// End FImageWrapperBase interface
+	virtual bool SetCompressed( const void* InCompressedData, int32 InCompressedSize ) override;
+	virtual bool SetRaw( const void* InRawData, int32 InRawSize, const int32 InWidth, const int32 InHeight, const ERGBFormat::Type InFormat, const int32 InBitDepth ) override;
+	virtual void Uncompress( const ERGBFormat::Type InFormat, int32 InBitDepth ) override;
+	virtual void Compress(int32 Quality) override;
 };
 
 
diff --git a/Engine/Source/Developer/ImageWrapper/Private/PngImageWrapper.cpp b/Engine/Source/Developer/ImageWrapper/Private/PngImageWrapper.cpp
--- a/Engine/Source/Developer/ImageWrapper/Private/PngImageWrapper.cpp
+++ b/Engine/Source/Developer/ImageWrapper/Private/PngImageWrapper.cpp
@@ -1,20 +1,8 @@
 // Copyright 1998-2014 Epic Games, Inc. All Rights Reserved.
 
-/*=============================================================================
-	PngImageWrapper.cpp: Implements the FPngImageWrapper class.
-
-	http://www.libpng.org/pub/png/libpng-1.2.5-manual.html
-
-	The below code isn't documented as it's black box behavior is almost 
-	entirely based on the sample from the libPNG documentation. 
-	
-	InitCompressed and InitRaw will set initial state and you will then be 
-	able to fill in Raw or CompressedData by calling Uncompress or Compress 
-	respectively.
-=============================================================================*/
-
 #include "ImageWrapperPrivatePCH.h"
 
+
 #if WITH_UNREALPNG
 
 
@@ -26,7 +14,7 @@ FCriticalSection GPNGSection;
  *****************************************************************************/
 
 /**
- * Guard that safely releases PNG reader resources
+ * Guard that safely releases PNG reader resources.
  */
 class PNGReadGuard
 {
diff --git a/Engine/Source/Developer/ImageWrapper/Private/PngImageWrapper.h b/Engine/Source/Developer/ImageWrapper/Private/PngImageWrapper.h
--- a/Engine/Source/Developer/ImageWrapper/Private/PngImageWrapper.h
+++ b/Engine/Source/Developer/ImageWrapper/Private/PngImageWrapper.h
@@ -1,9 +1,5 @@
 // Copyright 1998-2014 Epic Games, Inc. All Rights Reserved.
 
-/*=============================================================================
-	PngImageWrapper.h: Declares the FPngImageWrapper class.
-=============================================================================*/
-
 #pragma once
 
 #if WITH_UNREALPNG
@@ -15,7 +11,13 @@
 
 
 /**
- * PNG implementation of the helper class
+ * PNG implementation of the helper class.
+ *
+ * The implementation of this class is almost entirely based on the sample from the libPNG documentation.
+ * See http://www.libpng.org/pub/png/libpng-1.2.5-manual.html for details.
+ *	
+ * InitCompressed and InitRaw will set initial state and you will then be able to fill in Raw or
+ * CompressedData by calling Uncompress or Compress respectively.
  */
 class FPngImageWrapper
 	: public FImageWrapperBase
@@ -27,7 +29,6 @@ public:
 	 */
 	FPngImageWrapper();
 
-
 public:
 
 	// Begin FImageWrapper Interface
@@ -42,6 +43,8 @@ public:
 	
 	// End FImageWrapper Interface
 
+public:
+
 	/** 
 	 * Query whether this is a valid PNG type.
 	 *
@@ -70,7 +73,6 @@ protected:
 	static void* user_malloc( png_structp png_ptr, png_size_t size);
 	static void  user_free(png_structp png_ptr, png_voidp struct_ptr );
 
-
 private:
 
 	// The read offset into our array.
diff --git a/Engine/Source/Developer/ImageWrapper/Public/ImageWrapper.h b/Engine/Source/Developer/ImageWrapper/Public/ImageWrapper.h
--- a/Engine/Source/Developer/ImageWrapper/Public/ImageWrapper.h
+++ b/Engine/Source/Developer/ImageWrapper/Public/ImageWrapper.h
@@ -1,9 +1,5 @@
 // Copyright 1998-2014 Epic Games, Inc. All Rights Reserved.
 
-/*=============================================================================
-	ImageWrapper.h: ImageWrapper module public header file.
-=============================================================================*/
-
 #pragma once
 
 
diff --git a/Engine/Source/Developer/ImageWrapper/Public/Interfaces/IImageWrapper.h b/Engine/Source/Developer/ImageWrapper/Public/Interfaces/IImageWrapper.h
--- a/Engine/Source/Developer/ImageWrapper/Public/Interfaces/IImageWrapper.h
+++ b/Engine/Source/Developer/ImageWrapper/Public/Interfaces/IImageWrapper.h
@@ -1,9 +1,5 @@
 // Copyright 1998-2014 Epic Games, Inc. All Rights Reserved.
 
-/*=============================================================================
-	IImageWrapper.h: Declares the IImageWrapper interface.
-=============================================================================*/
-
 #pragma once
 
 
@@ -14,10 +10,17 @@ namespace EImageFormat
 	 */
 	enum Type
 	{
-		PNG,	// Portable Network Graphics
-		JPEG,	// Joint Photographic Experts Group
-		BMP,	// Windows Bitmap
-		ICO,	// Windows Icon resource
+		/** Portable Network Graphics */
+		PNG,
+
+		/** Joint Photographic Experts Group */
+		JPEG,
+
+		/** Windows Bitmap */
+		BMP,
+
+		/** Windows Icon resource */
+		ICO,
 	};
 };
 
@@ -37,14 +40,10 @@ namespace ERGBFormat
 };
 
 
-/**
- * Type definition for shared pointers to instances of IImageWrapper.
- */
+/** Type definition for shared pointers to instances of IImageWrapper. */
 typedef TSharedPtr<class IImageWrapper> IImageWrapperPtr;
 
-/**
- * Type definition for shared references to instances of IImageWrapper.
- */
+/** Type definition for shared references to instances of IImageWrapper. */
 typedef TSharedRef<class IImageWrapper> IImageWrapperRef;
 
 
@@ -60,7 +59,6 @@ public:
 	 *
 	 * @param InCompressedData - The memory address of the start of the compressed data
 	 * @param InCompressedSize - The size of the compressed data parsed
-	 *
 	 * @return	true if data was the expected format
 	 */
 	virtual bool SetCompressed( const void* InCompressedData, int32 InCompressedSize ) = 0;
@@ -74,7 +72,6 @@ public:
 	 * @param InHeight - the height of the image data
 	 * @param InFormat - the format the raw data is in, normally RGBA
 	 * @param InBitDepth - the bit-depth per channel, normally 8
-	 *
 	 * @return true if data was the expected format
 	 */
 	virtual bool SetRaw( const void* InRawData, int32 InRawSize, const int32 InWidth, const int32 InHeight, const ERGBFormat::Type InFormat, const int32 InBitDepth ) = 0;
@@ -92,7 +89,6 @@ public:
 	 * @param InFormat -  How we want to manipulate the RGB data
 	 * @param InBitDepth - The output bit-depth per channel, normally 8
 	 * @param OutRawData - Pointer to the uncompressed raw data.
-	 *
 	 * @return Array of the raw data
 	 */
 	virtual bool GetRaw( const ERGBFormat::Type InFormat, int32 InBitDepth, const TArray<uint8>*& OutRawData ) = 0;
diff --git a/Engine/Source/Developer/ImageWrapper/Public/Interfaces/IImageWrapperModule.h b/Engine/Source/Developer/ImageWrapper/Public/Interfaces/IImageWrapperModule.h
--- a/Engine/Source/Developer/ImageWrapper/Public/Interfaces/IImageWrapperModule.h
+++ b/Engine/Source/Developer/ImageWrapper/Public/Interfaces/IImageWrapperModule.h
@@ -1,9 +1,5 @@
 // Copyright 1998-2014 Epic Games, Inc. All Rights Reserved.
 
-/*=============================================================================
-	IImageWrapperModule.h: Declares the IImageWrapperModule interface.
-=============================================================================*/
-
 #pragma once
 
 
@@ -19,12 +15,10 @@ public:
 	 * Create a helper of a specific type 
 	 *
 	 * @param InFormat - The type of image we want to deal with
-	 *
 	 * @return The helper base class to manage the data
 	 */
 	virtual IImageWrapperPtr CreateImageWrapper( const EImageFormat::Type InFormat ) = 0;
 
-
 public:
 
 	/**
diff --git a/Engine/Source/Developer/InternationalizationSettings/Classes/InternationalizationSettingsModel.h b/Engine/Source/Developer/InternationalizationSettings/Classes/InternationalizationSettingsModel.h
--- a/Engine/Source/Developer/InternationalizationSettings/Classes/InternationalizationSettingsModel.h
+++ b/Engine/Source/Developer/InternationalizationSettings/Classes/InternationalizationSettingsModel.h
@@ -1,13 +1,10 @@
 // Copyright 1998-2014 Epic Games, Inc. All Rights Reserved.
 
-/*=============================================================================
-	InternationalizationSettingsModel.h: Declares the UInternationalizationSettingsModel class.
-=============================================================================*/
-
 #pragma once
 
 #include "InternationalizationSettingsModel.generated.h"
 
+
 /**
  * Implements loading and saving of internationalization settings.
  */
@@ -18,6 +15,7 @@ class INTERNATIONALIZATIONSETTINGS_API UInternationalizationSettingsModel
 	GENERATED_UCLASS_BODY()
 
 public:
+
 	void SaveDefaults();
 	void ResetToDefault();
 	FString GetCultureName() const;
@@ -26,6 +24,7 @@ public:
 	void ShouldLoadLocalizedPropertyNames(const bool Value);
 
 public:
+
 	/**
 	 * Returns an event delegate that is executed when a setting has changed.
 	 *
@@ -35,6 +34,7 @@ public:
 	FSettingChangedEvent& OnSettingChanged( ) { return SettingChangedEvent; }
 
 private:
+
 	// Holds an event delegate that is executed when a setting has changed.
 	FSettingChangedEvent SettingChangedEvent;
-};
\ No newline at end of file
+};
diff --git a/Engine/Source/Developer/InternationalizationSettings/Private/InternationalizationSettingsModel.cpp b/Engine/Source/Developer/InternationalizationSettings/Private/InternationalizationSettingsModel.cpp
--- a/Engine/Source/Developer/InternationalizationSettings/Private/InternationalizationSettingsModel.cpp
+++ b/Engine/Source/Developer/InternationalizationSettings/Private/InternationalizationSettingsModel.cpp
@@ -1,18 +1,15 @@
 // Copyright 1998-2014 Epic Games, Inc. All Rights Reserved.
 
-/*=============================================================================
-	Settings.cpp: Implements the constructors for the various settings classes.
-=============================================================================*/
-
 #include "InternationalizationSettingsModulePrivatePCH.h"
 
+
 /* UInternationalizationSettingsModel interface
  *****************************************************************************/
 
 UInternationalizationSettingsModel::UInternationalizationSettingsModel( const class FPostConstructInitializeProperties& PCIP )
 	: Super(PCIP)
-{
-}
+{ }
+
 
 void UInternationalizationSettingsModel::SaveDefaults()
 {
@@ -25,6 +22,7 @@ void UInternationalizationSettingsModel::SaveDefaults()
 	GConfig->SetBool( TEXT("Internationalization"), TEXT("ShouldLoadLocalizedPropertyNames"), bShouldLoadLocalizedPropertyNames, GEngineIni );
 }
 
+
 void UInternationalizationSettingsModel::ResetToDefault()
 {
 	FString SavedCultureName;
@@ -38,6 +36,7 @@ void UInternationalizationSettingsModel::ResetToDefault()
 	SettingChangedEvent.Broadcast();
 }
 
+
 FString UInternationalizationSettingsModel::GetCultureName() const
 {
 	FString SavedCultureName;
@@ -48,6 +47,7 @@ FString UInternationalizationSettingsModel::GetCultureName() const
 	return SavedCultureName;
 }
 
+
 void UInternationalizationSettingsModel::SetCultureName(const FString& CultureName)
 {
 	GConfig->SetString( TEXT("Internationalization"), TEXT("Culture"), *CultureName, GEditorGameAgnosticIni );
@@ -65,8 +65,9 @@ bool UInternationalizationSettingsModel::ShouldLoadLocalizedPropertyNames() cons
 	return bShouldLoadLocalizedPropertyNames;
 }
 
+
 void UInternationalizationSettingsModel::ShouldLoadLocalizedPropertyNames(const bool Value)
 {
 	GConfig->SetBool( TEXT("Internationalization"), TEXT("ShouldLoadLocalizedPropertyNames"), Value, GEditorGameAgnosticIni );
 	SettingChangedEvent.Broadcast();
-}
\ No newline at end of file
+}
diff --git a/Engine/Source/Developer/InternationalizationSettings/Public/InternationalizationSettings.h b/Engine/Source/Developer/InternationalizationSettings/Public/InternationalizationSettings.h
--- a/Engine/Source/Developer/InternationalizationSettings/Public/InternationalizationSettings.h
+++ b/Engine/Source/Developer/InternationalizationSettings/Public/InternationalizationSettings.h
@@ -1,16 +1,11 @@
 // Copyright 1998-2014 Epic Games, Inc. All Rights Reserved.
 
-/*=============================================================================
-	AutomationWindow.h: AutomationWindow module public header file.
-=============================================================================*/
-
 #pragma once
 
 
 /* Dependencies
  *****************************************************************************/
 
-//#include "Slate.h"
 #include "EditorStyle.h"
 #include "InternationalizationSettingsModel.h"
-#include "InternationalizationSettingsModelDetails.h"
\ No newline at end of file
+#include "InternationalizationSettingsModelDetails.h"
diff --git a/Engine/Source/Developer/LauncherServices/Private/Launcher/Launcher.cpp b/Engine/Source/Developer/LauncherServices/Private/Launcher/Launcher.cpp
--- a/Engine/Source/Developer/LauncherServices/Private/Launcher/Launcher.cpp
+++ b/Engine/Source/Developer/LauncherServices/Private/Launcher/Launcher.cpp
@@ -1,9 +1,5 @@
 // Copyright 1998-2014 Epic Games, Inc. All Rights Reserved.
 
-/*=============================================================================
-	Launcher.cpp: Implements the FLauncher class.
-=============================================================================*/
-
 #include "LauncherServicesPrivatePCH.h"
 
 
@@ -16,11 +12,11 @@ ILauncherWorkerPtr FLauncher::Launch( const ITargetDeviceProxyManagerRef& Device
 	{
 		FLauncherWorker* LauncherWorker = new FLauncherWorker(DeviceProxyManager, Profile);
 
-		if ((LauncherWorker != NULL) && (FRunnableThread::Create(LauncherWorker, TEXT("LauncherWorker")) != NULL))
+		if ((LauncherWorker != nullptr) && (FRunnableThread::Create(LauncherWorker, TEXT("LauncherWorker")) != nullptr))
 		{
 			return MakeShareable(LauncherWorker);
 		}			
 	}
 
-	return NULL;
+	return nullptr;
 }
diff --git a/Engine/Source/Developer/LauncherServices/Private/Launcher/Launcher.h b/Engine/Source/Developer/LauncherServices/Private/Launcher/Launcher.h
--- a/Engine/Source/Developer/LauncherServices/Private/Launcher/Launcher.h
+++ b/Engine/Source/Developer/LauncherServices/Private/Launcher/Launcher.h
@@ -1,9 +1,5 @@
 // Copyright 1998-2014 Epic Games, Inc. All Rights Reserved.
 
-/*=============================================================================
-	Launcher.h: Declares the FLauncher class.
-=============================================================================*/
-
 #pragma once
 
 
diff --git a/Engine/Source/Developer/LauncherServices/Private/Launcher/LauncherBuildCommands.h b/Engine/Source/Developer/LauncherServices/Private/Launcher/LauncherBuildCommands.h
--- a/Engine/Source/Developer/LauncherServices/Private/Launcher/LauncherBuildCommands.h
+++ b/Engine/Source/Developer/LauncherServices/Private/Launcher/LauncherBuildCommands.h
@@ -1,9 +1,5 @@
 // Copyright 1998-2014 Epic Games, Inc. All Rights Reserved.
 
-/*=============================================================================
-	LauncherBuildCommands.h: Declares the FLauncherBuildCommands.
-=============================================================================*/
-
 #pragma once
 
 
@@ -14,10 +10,10 @@ class FLauncherBuildGameCommand
 	: public FLauncherUATCommand
 {
 public:
+
 	FLauncherBuildGameCommand( const ITargetPlatform& InTargetPlatform )
 		: TargetPlatform(InTargetPlatform)
-	{
-	}
+	{ }
 
 	virtual FString GetName() const OVERRIDE
 	{
@@ -52,10 +48,10 @@ class FLauncherBuildServerCommand
 	: public FLauncherUATCommand
 {
 public:
+
 	FLauncherBuildServerCommand( const ITargetPlatform& InTargetPlatform )
 		: TargetPlatform(InTargetPlatform)
-	{
-	}
+	{ }
 
 	virtual FString GetName() const OVERRIDE
 	{
diff --git a/Engine/Source/Developer/LauncherServices/Private/Launcher/LauncherCookCommands.h b/Engine/Source/Developer/LauncherServices/Private/Launcher/LauncherCookCommands.h
--- a/Engine/Source/Developer/LauncherServices/Private/Launcher/LauncherCookCommands.h
+++ b/Engine/Source/Developer/LauncherServices/Private/Launcher/LauncherCookCommands.h
@@ -1,9 +1,5 @@
 // Copyright 1998-2014 Epic Games, Inc. All Rights Reserved.
 
-/*=============================================================================
-	LauncherCookCommands.h: Declares the FLauncherCookCommands.
-=============================================================================*/
-
 #pragma once
 
 
@@ -14,10 +10,10 @@ class FLauncherCookGameCommand
 	: public FLauncherUATCommand
 {
 public:
+
 	FLauncherCookGameCommand( const ITargetPlatform& InTargetPlatform )
 		: TargetPlatform(InTargetPlatform)
-	{
-	}
+	{ }
 
 	virtual FString GetName() const OVERRIDE
 	{
@@ -95,10 +91,10 @@ class FLauncherCookServerCommand
 	: public FLauncherUATCommand
 {
 public:
+
 	FLauncherCookServerCommand( const ITargetPlatform& InTargetPlatform )
 		: TargetPlatform(InTargetPlatform)
-	{
-	}
+	{ }
 
 	virtual FString GetName() const OVERRIDE
 	{
@@ -186,6 +182,7 @@ class FLauncherCookOnTheFlyCommand
 	: public FLauncherUATCommand
 {
 public:
+
 	FLauncherCookOnTheFlyCommand( const ITargetPlatform& InTargetPlatform )
 		: TargetPlatform(InTargetPlatform)
 		, InstanceId(FGuid::NewGuid())
@@ -320,10 +317,10 @@ class FLauncherStandAloneCookOnTheFlyCommand
 	: public FLauncherUATCommand
 {
 public:
+
 	FLauncherStandAloneCookOnTheFlyCommand( const ITargetPlatform& InTargetPlatform )
 		: TargetPlatform(InTargetPlatform)
-	{
-	}
+	{ }
 
 	virtual FString GetName() const OVERRIDE
 	{
diff --git a/Engine/Source/Developer/LauncherServices/Private/Launcher/LauncherDeployCommands.h b/Engine/Source/Developer/LauncherServices/Private/Launcher/LauncherDeployCommands.h
--- a/Engine/Source/Developer/LauncherServices/Private/Launcher/LauncherDeployCommands.h
+++ b/Engine/Source/Developer/LauncherServices/Private/Launcher/LauncherDeployCommands.h
@@ -1,9 +1,5 @@
 // Copyright 1998-2014 Epic Games, Inc. All Rights Reserved.
 
-/*=============================================================================
-	LauncherDeployCommands.h: Declares the FLauncherDeployCommands.
-=============================================================================*/
-
 #pragma once
 
 
@@ -20,8 +16,7 @@ public:
 		, InstanceId(FGuid::NewGuid())
 		, CookCommand(InCook)
 		, LauncherCommandLine(InCmdLine)
-	{
-	}
+	{ }
 
 	virtual FString GetName() const OVERRIDE
 	{
@@ -122,13 +117,13 @@ class FLauncherDeployServerToDeviceCommand
 	: public FLauncherUATCommand
 {
 public:
+
 	FLauncherDeployServerToDeviceCommand( const ITargetDeviceProxyRef& InDeviceProxy, const ITargetPlatform& InTargetPlatform, const TSharedPtr<FLauncherUATCommand>& InCook )
 		: DeviceProxy(InDeviceProxy)
 		, TargetPlatform(InTargetPlatform)
 		, InstanceId(FGuid::NewGuid())
 		, CookCommand(InCook)
-	{
-	}
+	{ }
 
 	virtual FString GetName() const OVERRIDE
 	{
@@ -218,14 +213,14 @@ class FLauncherDeployGamePackageToDeviceCommand
 	: public FLauncherUATCommand
 {
 public:
+
 	FLauncherDeployGamePackageToDeviceCommand( const ITargetDeviceProxyRef& InDeviceProxy, const ITargetPlatform& InTargetPlatform, const TSharedPtr<FLauncherUATCommand>& InCook, const FString& InCmdLine )
 		: DeviceProxy(InDeviceProxy)
 		, TargetPlatform(InTargetPlatform)
 		, InstanceId(FGuid::NewGuid())
 		, CookCommand(InCook)
 		, LauncherCommandLine(InCmdLine)
-	{
-	}
+	{ }
 
 	virtual FString GetName() const OVERRIDE
 	{
@@ -313,13 +308,13 @@ class FLauncherDeployServerPackageToDeviceCommand
 	: public FLauncherUATCommand
 {
 public:
+
 	FLauncherDeployServerPackageToDeviceCommand( const ITargetDeviceProxyRef& InDeviceProxy, const ITargetPlatform& InTargetPlatform, const TSharedPtr<FLauncherUATCommand>& InCook )
 		: DeviceProxy(InDeviceProxy)
 		, TargetPlatform(InTargetPlatform)
 		, InstanceId(FGuid::NewGuid())
 		, CookCommand(InCook)
-	{
-	}
+	{ }
 
 	virtual FString GetName() const OVERRIDE
 	{
diff --git a/Engine/Source/Developer/LauncherServices/Private/Launcher/LauncherLaunchCommands.h b/Engine/Source/Developer/LauncherServices/Private/Launcher/LauncherLaunchCommands.h
--- a/Engine/Source/Developer/LauncherServices/Private/Launcher/LauncherLaunchCommands.h
+++ b/Engine/Source/Developer/LauncherServices/Private/Launcher/LauncherLaunchCommands.h
@@ -1,9 +1,5 @@
 // Copyright 1998-2014 Epic Games, Inc. All Rights Reserved.
 
-/*=============================================================================
-	LauncherLaunchCommands.h: Declares the FLauncherLaunchCommands.
-=============================================================================*/
-
 #pragma once
 
 
@@ -14,14 +10,14 @@ class FLauncherLaunchGameCommand
 	: public FLauncherUATCommand
 {
 public:
+
 	FLauncherLaunchGameCommand(const ITargetDeviceProxyRef& InDeviceProxy, const ITargetPlatform& InTargetPlatform, const ILauncherProfileLaunchRoleRef& InRole, const TSharedPtr<FLauncherUATCommand>& InCook)
 		: DeviceProxy(InDeviceProxy)
 		, TargetPlatform(InTargetPlatform)
 		, InstanceId(FGuid::NewGuid())
 		, Role(InRole)
 		, CookCommand(InCook)
-	{
-	}
+	{ }
 
 	virtual FString GetName() const OVERRIDE
 	{
@@ -79,6 +75,7 @@ public:
 	}
 
 private:
+
 	// Holds the device proxy to launch on.
 	ITargetDeviceProxyPtr DeviceProxy;
 
@@ -103,14 +100,14 @@ class FLauncherLaunchDedicatedServerCommand
 	: public FLauncherUATCommand
 {
 public:
+
 	FLauncherLaunchDedicatedServerCommand(const ITargetDeviceProxyRef& InDeviceProxy, const ITargetPlatform& InTargetPlatform, const ILauncherProfileLaunchRoleRef& InRole, const TSharedPtr<FLauncherUATCommand>& InCook)
 		: DeviceProxy(InDeviceProxy)
 		, TargetPlatform(InTargetPlatform)
 		, InstanceId(FGuid::NewGuid())
 		, Role(InRole)
 		, CookCommand(InCook)
-	{
-	}
+	{ }
 
 	virtual FString GetName() const OVERRIDE
 	{
@@ -177,6 +174,7 @@ public:
 	}
 
 private:
+
 	// Holds the device proxy to launch on.
 	ITargetDeviceProxyPtr DeviceProxy;
 
@@ -189,6 +187,6 @@ private:
 	// Holds the role of the instance to launch.
 	ILauncherProfileLaunchRoleRef Role;
 
-	// cook command used for this build
+	// Cook command used for this build
 	const TSharedPtr<FLauncherUATCommand> CookCommand;
 };
diff --git a/Engine/Source/Developer/LauncherServices/Private/Launcher/LauncherPackageCommands.h b/Engine/Source/Developer/LauncherServices/Private/Launcher/LauncherPackageCommands.h
--- a/Engine/Source/Developer/LauncherServices/Private/Launcher/LauncherPackageCommands.h
+++ b/Engine/Source/Developer/LauncherServices/Private/Launcher/LauncherPackageCommands.h
@@ -1,9 +1,5 @@
 // Copyright 1998-2014 Epic Games, Inc. All Rights Reserved.
 
-/*=============================================================================
-	LauncherPackageCommands.h: Declares the FLauncherPackageCommands.
-=============================================================================*/
-
 #pragma once
 
 
diff --git a/Engine/Source/Developer/LauncherServices/Private/Launcher/LauncherTask.h b/Engine/Source/Developer/LauncherServices/Private/Launcher/LauncherTask.h
--- a/Engine/Source/Developer/LauncherServices/Private/Launcher/LauncherTask.h
+++ b/Engine/Source/Developer/LauncherServices/Private/Launcher/LauncherTask.h
@@ -1,9 +1,5 @@
 // Copyright 1998-2014 Epic Games, Inc. All Rights Reserved.
 
-/*=============================================================================
-	LauncherTask.h: Declares the FLauncherTask class.
-=============================================================================*/
-
 #pragma once
 
 
@@ -30,7 +26,6 @@ public:
 		, Status(ELauncherTaskStatus::Pending)
 	{ }
 
-
 public:
 
 	/**
@@ -61,7 +56,7 @@ public:
 
 		Thread = FRunnableThread::Create(this, TEXT("FLauncherTask"));
 
-		if (Thread == NULL)
+		if (Thread == nullptr)
 		{
 
 		}
@@ -173,7 +168,6 @@ public:
 
 	// End FRunnable interface
 
-
 public:
 
 	// Begin ILauncherTask interface
@@ -241,7 +235,6 @@ public:
 
 	// End ILauncherTask interface
 
-
 protected:
 
 	/**
@@ -253,7 +246,6 @@ protected:
 	 */
 	virtual bool PerformTask( FLauncherTaskChainState& ChainState ) = 0;
 
-
 protected:
 
 	/**
@@ -278,7 +270,6 @@ protected:
 		}
 	}
 
-
 private:
 
 	// Holds the tasks that execute after this task completed.
diff --git a/Engine/Source/Developer/LauncherServices/Private/Launcher/LauncherTaskChainState.h b/Engine/Source/Developer/LauncherServices/Private/Launcher/LauncherTaskChainState.h
--- a/Engine/Source/Developer/LauncherServices/Private/Launcher/LauncherTaskChainState.h
+++ b/Engine/Source/Developer/LauncherServices/Private/Launcher/LauncherTaskChainState.h
@@ -1,9 +1,5 @@
 // Copyright 1998-2014 Epic Games, Inc. All Rights Reserved.
 
-/*=============================================================================
-	LauncherTaskChainState.h: Declares the FLauncherTaskChainState class.
-=============================================================================*/
-
 #pragma once
 
 
@@ -12,33 +8,21 @@
  */
 struct FLauncherTaskChainState
 {
-	/**
-	 * Holds the path to the binaries directory (for deployment).
-	 */
+	/** The path to the binaries directory (for deployment). */
 	FString BinariesDirectory;
 
-	/**
-	 * Holds the path to the content directory (for deployment).
-	 */
+	/** The path to the content directory (for deployment). */
 	FString ContentDirectory;
 
-	/**
-	 * Holds the identifier of the deployed application.
-	 */
+	/** The identifier of the deployed application. */
 	FString DeployedAppId;
 
-	/**
-	 * Holds the list of IP addresses and port numbers at which the file server is reachable.
-	 */
+	/** The list of IP addresses and port numbers at which the file server is reachable. */
 	FString FileServerAddressListString;
 
-	/**
-	 * Holds a pointer to the launcher profile.
-	 */
+	/** A pointer to the launcher profile. */
 	ILauncherProfilePtr Profile;
 
-	/**
-	 * Holds the session identifier.
-	 */
+	/** The session identifier. */
 	FGuid SessionId;
-};
\ No newline at end of file
+};
diff --git a/Engine/Source/Developer/LauncherServices/Private/Launcher/LauncherUATCommand.h b/Engine/Source/Developer/LauncherServices/Private/Launcher/LauncherUATCommand.h
--- a/Engine/Source/Developer/LauncherServices/Private/Launcher/LauncherUATCommand.h
+++ b/Engine/Source/Developer/LauncherServices/Private/Launcher/LauncherUATCommand.h
@@ -1,11 +1,8 @@
 // Copyright 1998-2014 Epic Games, Inc. All Rights Reserved.
 
-/*=============================================================================
-	FLauncherUATCommand.h: Declares the FLauncherUATCommand class.
-=============================================================================*/
-
 #pragma once
 
+
 /**
  * class for setting up UAT command arguments for various launcher commands
  */
@@ -95,4 +92,3 @@ public:
 		return true;
 	}
 };
-
diff --git a/Engine/Source/Developer/LauncherServices/Private/Launcher/LauncherUATTask.h b/Engine/Source/Developer/LauncherServices/Private/Launcher/LauncherUATTask.h
--- a/Engine/Source/Developer/LauncherServices/Private/Launcher/LauncherUATTask.h
+++ b/Engine/Source/Developer/LauncherServices/Private/Launcher/LauncherUATTask.h
@@ -1,13 +1,10 @@
 // Copyright 1998-2014 Epic Games, Inc. All Rights Reserved.
 
-/*=============================================================================
-	FLauncherUATTask.h: Declares the FLauncherUATTask class.
-=============================================================================*/
-
 #pragma once
 
 #include "PlatformInfo.h"
 
+
 /**
  * class for UAT launcher tasks.
  */
@@ -33,7 +30,6 @@ public:
 		NoCompile = !FParse::Param( FCommandLine::Get(), TEXT("development") ) ? TEXT(" -nocompile") : TEXT("");
 	}
 
-
 protected:
 
 	/**
@@ -137,7 +133,6 @@ protected:
 		return (ReturnCode == 0);
 	}
 
-
 private:
 
 	// command for this task
diff --git a/Engine/Source/Developer/LauncherServices/Private/Launcher/LauncherVerifyProfileTask.h b/Engine/Source/Developer/LauncherServices/Private/Launcher/LauncherVerifyProfileTask.h
--- a/Engine/Source/Developer/LauncherServices/Private/Launcher/LauncherVerifyProfileTask.h
+++ b/Engine/Source/Developer/LauncherServices/Private/Launcher/LauncherVerifyProfileTask.h
@@ -1,9 +1,5 @@
 // Copyright 1998-2014 Epic Games, Inc. All Rights Reserved.
 
-/*=============================================================================
-	LauncherVerifyProfileTask.h: Declares the FLauncherVerifyProfileTask class.
-=============================================================================*/
-
 #pragma once
 
 
@@ -26,7 +22,6 @@ public:
 		: FLauncherTask(NSLOCTEXT("FLauncherTask", "LauncherVerifyProfileName", "Verify").ToString(), NSLOCTEXT("FLauncherTask", "LauncherVerifyProfileDesc", "Verify profile settings").ToString(), NULL, NULL)
 	{ }
 
-
 protected:
 
 	virtual bool PerformTask( FLauncherTaskChainState& ChainState ) OVERRIDE
@@ -38,4 +33,4 @@ protected:
 		
 		return false;
 	}
-};
\ No newline at end of file
+};
diff --git a/Engine/Source/Developer/LauncherServices/Private/Launcher/LauncherWorker.cpp b/Engine/Source/Developer/LauncherServices/Private/Launcher/LauncherWorker.cpp
--- a/Engine/Source/Developer/LauncherServices/Private/Launcher/LauncherWorker.cpp
+++ b/Engine/Source/Developer/LauncherServices/Private/Launcher/LauncherWorker.cpp
@@ -1,9 +1,5 @@
 // Copyright 1998-2014 Epic Games, Inc. All Rights Reserved.
 
-/*=============================================================================
-	LauncherWorker.cpp: Implements the FLauncherWorker class.
-=============================================================================*/
-
 #include "LauncherServicesPrivatePCH.h"
 
 
@@ -164,16 +160,19 @@ int32 FLauncherWorker::GetTasks( TArray<ILauncherTaskPtr>& OutTasks ) const
 	return OutTasks.Num();
 }
 
+
 void FLauncherWorker::OnTaskStarted(const FString& TaskName)
 {
 	StageStarted.Broadcast(TaskName);
 }
 
+
 void FLauncherWorker::OnTaskCompleted(const FString& TaskName)
 {
 	StageCompleted.Broadcast(TaskName);
 }
 
+
 /* FLauncherWorker implementation
  *****************************************************************************/
 
@@ -222,7 +221,7 @@ void FLauncherWorker::CreateAndExecuteTasks( const ILauncherProfileRef& InProfil
 		const FString& TargetPlatformName = Platforms[PlatformIndex];
 		const ITargetPlatform* TargetPlatform = GetTargetPlatformManager()->FindTargetPlatform(TargetPlatformName);
 
-		if (TargetPlatform == NULL)
+		if (TargetPlatform == nullptr)
 		{
 			continue;
 		}
diff --git a/Engine/Source/Developer/LauncherServices/Private/Launcher/LauncherWorker.h b/Engine/Source/Developer/LauncherServices/Private/Launcher/LauncherWorker.h
--- a/Engine/Source/Developer/LauncherServices/Private/Launcher/LauncherWorker.h
+++ b/Engine/Source/Developer/LauncherServices/Private/Launcher/LauncherWorker.h
@@ -1,11 +1,8 @@
 // Copyright 1998-2014 Epic Games, Inc. All Rights Reserved.
 
-/*=============================================================================
-	LauncherWorker.h: Declares the FLauncherWorker class.
-=============================================================================*/
-
 #pragma once
 
+
 /**
  * Implements the launcher's worker thread.
  */
@@ -39,7 +36,6 @@ public:
 
 	virtual void Exit( ) OVERRIDE { }
 
-
 public:
 
 	virtual void Cancel( ) OVERRIDE;
diff --git a/Engine/Source/Developer/LauncherServices/Private/LauncherServicesModule.cpp b/Engine/Source/Developer/LauncherServices/Private/LauncherServicesModule.cpp
--- a/Engine/Source/Developer/LauncherServices/Private/LauncherServicesModule.cpp
+++ b/Engine/Source/Developer/LauncherServices/Private/LauncherServicesModule.cpp
@@ -1,9 +1,5 @@
 // Copyright 1998-2014 Epic Games, Inc. All Rights Reserved.
 
-/*=============================================================================
-	LauncherServicesModule.cpp: Implements the FLauncherServicesModule class.
-=============================================================================*/
-
 #include "LauncherServicesPrivatePCH.h"
 
 
@@ -59,7 +55,6 @@ public:
 
 	// End ILauncherServicesModule interface
 
-
 private:
 	
 	/// Event to be called when the editor tried to use a platform, but it wasn't installed
diff --git a/Engine/Source/Developer/LauncherServices/Private/LauncherServicesPrivatePCH.h b/Engine/Source/Developer/LauncherServices/Private/LauncherServicesPrivatePCH.h
--- a/Engine/Source/Developer/LauncherServices/Private/LauncherServicesPrivatePCH.h
+++ b/Engine/Source/Developer/LauncherServices/Private/LauncherServicesPrivatePCH.h
@@ -1,13 +1,8 @@
 // Copyright 1998-2014 Epic Games, Inc. All Rights Reserved.
 
-/*=============================================================================
-	LauncherServicesPrivatePCH.h: Pre-compiled header file for the LauncherServices module.
-=============================================================================*/
-
 #pragma once
 
-
-#include "../Public/LauncherServices.h"
+#include "LauncherServices.h"
 
 
 /* Dependencies
@@ -25,9 +20,7 @@
 /* Private includes
  *****************************************************************************/
 
-/**
- * Defines the launcher profile file version.
- */
+/** Defines the launcher profile file version. */
 #define LAUNCHERSERVICES_PROFILEVERSION 8
 
 
@@ -49,5 +42,4 @@
 #include "LauncherPackageCommands.h"
 #include "LauncherVerifyProfileTask.h"
 #include "LauncherWorker.h"
-
 #include "Launcher.h"
diff --git a/Engine/Source/Developer/LauncherServices/Private/Profiles/LauncherDeviceGroup.h b/Engine/Source/Developer/LauncherServices/Private/Profiles/LauncherDeviceGroup.h
--- a/Engine/Source/Developer/LauncherServices/Private/Profiles/LauncherDeviceGroup.h
+++ b/Engine/Source/Developer/LauncherServices/Private/Profiles/LauncherDeviceGroup.h
@@ -1,9 +1,5 @@
 // Copyright 1998-2014 Epic Games, Inc. All Rights Reserved.
 
-/*=============================================================================
-	LauncherDeviceGroup.h: Declares the FLauncherDeviceGroup class.
-=============================================================================*/
-
 #pragma once
 
 
@@ -32,7 +28,6 @@ public:
 		, Name(InName)
 	{ }
 
-
 public:
 
 	// Begin ILauncherDeviceGroup interface
@@ -88,7 +83,6 @@ public:
 
 	// End ILauncherDeviceGroup interface
 
-
 private:
 
 	// Holds the devices that are part of this group.
@@ -100,7 +94,6 @@ private:
 	// Holds the human readable name of this group.
 	FString Name;
 
-
 private:
 
 	// Holds a delegate to be invoked when a device was added to this group.
diff --git a/Engine/Source/Developer/LauncherServices/Private/Profiles/LauncherProfile.h b/Engine/Source/Developer/LauncherServices/Private/Profiles/LauncherProfile.h
--- a/Engine/Source/Developer/LauncherServices/Private/Profiles/LauncherProfile.h
+++ b/Engine/Source/Developer/LauncherServices/Private/Profiles/LauncherProfile.h
@@ -1,9 +1,5 @@
 // Copyright 1998-2014 Epic Games, Inc. All Rights Reserved.
 
-/*=============================================================================
-	LauncherProfile.h: Declares the FLauncherProfile class.
-=============================================================================*/
-
 #pragma once
 
 
@@ -40,7 +36,6 @@ public:
 	 */
 	~FLauncherProfile( ) { }
 
-
 public:
 
 	/**
@@ -57,7 +52,6 @@ public:
 		return DeployedDeviceGroupId;
 	}
 
-
 public:
 
 	// Begin ILauncherProfile interface
@@ -686,7 +680,6 @@ public:
 
 	// End ILauncherProfile interface
 
-
 protected:
 
 	/**
@@ -840,7 +833,6 @@ protected:
 		}
 	}
 
-
 private:
 
 	// Holds the desired build configuration (only used if creating new builds).
diff --git a/Engine/Source/Developer/LauncherServices/Private/Profiles/LauncherProfileLaunchRole.h b/Engine/Source/Developer/LauncherServices/Private/Profiles/LauncherProfileLaunchRole.h
--- a/Engine/Source/Developer/LauncherServices/Private/Profiles/LauncherProfileLaunchRole.h
+++ b/Engine/Source/Developer/LauncherServices/Private/Profiles/LauncherProfileLaunchRole.h
@@ -1,9 +1,5 @@
 // Copyright 1998-2014 Epic Games, Inc. All Rights Reserved.
 
-/*=============================================================================
-	LauncherProfileLaunchRole.h: Declares the FLauncherProfileLaunchRole class.
-=============================================================================*/
-
 #pragma once
 
 
@@ -12,9 +8,7 @@ class FLauncherProfileLaunchRole
 {
 public:
 
-	/**
-	 * Default constructor.
-	 */
+	/** Default constructor. */
 	FLauncherProfileLaunchRole(  )
 		: InstanceType(ELauncherProfileRoleInstanceTypes::StandaloneClient)
 		, Name(TEXT("Unnamed Role"))
@@ -30,7 +24,6 @@ public:
 		Serialize(Archive);
 	}
 
-
 public:
 
 	// Begin ILauncherProfileLaunchRole interface
@@ -114,7 +107,6 @@ public:
 
 	// End ILauncherProfileLaunchRole interface
 
-
 private:
 
 	// Holds the identifier of the device that is assigned to this role.
@@ -140,4 +132,4 @@ private:
 
 	// Holds a flag indicating whether VSync should be enabled.
 	bool VsyncEnabled;
-};
\ No newline at end of file
+};
diff --git a/Engine/Source/Developer/LauncherServices/Private/Profiles/LauncherProfileManager.cpp b/Engine/Source/Developer/LauncherServices/Private/Profiles/LauncherProfileManager.cpp
--- a/Engine/Source/Developer/LauncherServices/Private/Profiles/LauncherProfileManager.cpp
+++ b/Engine/Source/Developer/LauncherServices/Private/Profiles/LauncherProfileManager.cpp
@@ -1,9 +1,5 @@
 // Copyright 1998-2014 Epic Games, Inc. All Rights Reserved.
 
-/*=============================================================================
-	LauncherProfileManager.cpp: Implements the FLauncherProfileManager class.
-=============================================================================*/
-
 #include "LauncherServicesPrivatePCH.h"
 
 
@@ -110,7 +106,7 @@ ILauncherProfilePtr FLauncherProfileManager::FindProfile( const FString& Profile
 		}
 	}
 
-	return NULL;
+	return nullptr;
 }
 
 
@@ -138,7 +134,7 @@ ILauncherDeviceGroupPtr FLauncherProfileManager::GetDeviceGroup( const FGuid& Gr
 		}
 	}
 
-	return NULL;
+	return nullptr;
 }
 
 
@@ -154,7 +150,7 @@ ILauncherProfilePtr FLauncherProfileManager::GetProfile( const FGuid& ProfileId
 		}
 	}
 
-	return NULL;
+	return nullptr;
 }
 
 
@@ -169,7 +165,7 @@ ILauncherProfilePtr FLauncherProfileManager::LoadProfile( FArchive& Archive )
 		return MakeShareable(Profile);
 	}
 
-	return NULL;
+	return nullptr;
 }
 
 
@@ -225,11 +221,11 @@ void FLauncherProfileManager::SaveSettings( )
 
 void FLauncherProfileManager::LoadDeviceGroups( )
 {
-	if (GConfig != NULL)
+	if (GConfig != nullptr)
 	{
 		FConfigSection* LoadedDeviceGroups = GConfig->GetSectionPrivate(TEXT("Launcher.DeviceGroups"), false, true, GEngineIni);
 
-		if (LoadedDeviceGroups != NULL)
+		if (LoadedDeviceGroups != nullptr)
 		{
 			// parse the configuration file entries into device groups
 			for (FConfigSection::TIterator It(*LoadedDeviceGroups); It; ++It)
@@ -255,7 +251,7 @@ void FLauncherProfileManager::LoadProfiles( )
 		FString ProfileFilePath = GetProfileFolder() / *It;
 		FArchive* ProfileFileReader = IFileManager::Get().CreateFileReader(*ProfileFilePath);
 
-		if (ProfileFileReader != NULL)
+		if (ProfileFileReader != nullptr)
 		{
 			ILauncherProfilePtr LoadedProfile = LoadProfile(*ProfileFileReader);
 			delete ProfileFileReader;
@@ -311,7 +307,7 @@ ILauncherDeviceGroupPtr FLauncherProfileManager::ParseDeviceGroup( const FString
 
 void FLauncherProfileManager::SaveDeviceGroups( )
 {
-	if (GConfig != NULL)
+	if (GConfig != nullptr)
 	{
 		GConfig->EmptySection(TEXT("Launcher.DeviceGroups"), GEngineIni);
 
@@ -354,7 +350,7 @@ void FLauncherProfileManager::SaveProfiles( )
 		FString ProfileFileName = GetProfileFolder() / (*It)->GetId().ToString() + TEXT(".ulp");
 		FArchive* ProfileFileWriter = IFileManager::Get().CreateFileWriter(*ProfileFileName);
 
-		if (ProfileFileWriter != NULL)
+		if (ProfileFileWriter != nullptr)
 		{
 			SaveProfile((*It).ToSharedRef(), *ProfileFileWriter);
 
diff --git a/Engine/Source/Developer/LauncherServices/Private/Profiles/LauncherProfileManager.h b/Engine/Source/Developer/LauncherServices/Private/Profiles/LauncherProfileManager.h
--- a/Engine/Source/Developer/LauncherServices/Private/Profiles/LauncherProfileManager.h
+++ b/Engine/Source/Developer/LauncherServices/Private/Profiles/LauncherProfileManager.h
@@ -1,9 +1,5 @@
 // Copyright 1998-2014 Epic Games, Inc. All Rights Reserved.
 
-/*=============================================================================
-	LauncherProfileManager.h: Declares the FLauncherProfileManager class.
-=============================================================================*/
-
 #pragma once
 
 
@@ -15,9 +11,7 @@ class FLauncherProfileManager
 {
 public:
 
-	/**
-	 * Default constructor.
-	 */
+	/** Default constructor. */
 	FLauncherProfileManager( );
 
 
@@ -77,7 +71,6 @@ public:
 
 	// End ILauncherProfileManager interface
 
-
 protected:
 
 	/*
@@ -109,7 +102,6 @@ protected:
 	 */
 	void SaveProfiles( );
 
-
 protected:
 
 	/**
@@ -122,7 +114,6 @@ protected:
 		return FPaths::EngineSavedDir() / TEXT("Launcher");
 	}
 
-
 private:
 
 	// Holds the collection of device groups.
@@ -131,7 +122,6 @@ private:
 	// Holds the collection of launcher profiles.
 	TArray<ILauncherProfilePtr> Profiles;
 
-
 private:
 
 	// Holds a delegate to be invoked when a device group was added.
diff --git a/Engine/Source/Developer/LauncherServices/Public/Common/GameProjectHelper.h b/Engine/Source/Developer/LauncherServices/Public/Common/GameProjectHelper.h
--- a/Engine/Source/Developer/LauncherServices/Public/Common/GameProjectHelper.h
+++ b/Engine/Source/Developer/LauncherServices/Public/Common/GameProjectHelper.h
@@ -1,9 +1,5 @@
 // Copyright 1998-2014 Epic Games, Inc. All Rights Reserved.
 
-/*=============================================================================
-	GameProjectHelper.h: Declares the FGameProjectHelper class.
-=============================================================================*/
-
 #pragma once
 
 
diff --git a/Engine/Source/Developer/LauncherServices/Public/Interfaces/ILauncher.h b/Engine/Source/Developer/LauncherServices/Public/Interfaces/ILauncher.h
--- a/Engine/Source/Developer/LauncherServices/Public/Interfaces/ILauncher.h
+++ b/Engine/Source/Developer/LauncherServices/Public/Interfaces/ILauncher.h
@@ -1,20 +1,12 @@
 // Copyright 1998-2014 Epic Games, Inc. All Rights Reserved.
 
-/*=============================================================================
-	ILauncher.h: Declares the ILauncher interface.
-=============================================================================*/
-
 #pragma once
 
 
-/**
- * Type definition for shared pointers to instances of ILauncher.
- */
+/** Type definition for shared pointers to instances of ILauncher. */
 typedef TSharedPtr<class ILauncher> ILauncherPtr;
 
-/**
- * Type definition for shared references to instances of ILauncher.
- */
+/** Type definition for shared references to instances of ILauncher. */
 typedef TSharedRef<class ILauncher> ILauncherRef;
 
 
@@ -35,7 +27,6 @@ public:
 	 */
 	virtual ILauncherWorkerPtr Launch( const ITargetDeviceProxyManagerRef& DeviceProxyManager, const ILauncherProfileRef& Profile ) = 0;
 
-
 public:
 
 	/**
diff --git a/Engine/Source/Developer/LauncherServices/Public/Interfaces/ILauncherDeviceGroup.h b/Engine/Source/Developer/LauncherServices/Public/Interfaces/ILauncherDeviceGroup.h
--- a/Engine/Source/Developer/LauncherServices/Public/Interfaces/ILauncherDeviceGroup.h
+++ b/Engine/Source/Developer/LauncherServices/Public/Interfaces/ILauncherDeviceGroup.h
@@ -1,19 +1,12 @@
 // Copyright 1998-2014 Epic Games, Inc. All Rights Reserved.
 
-/*=============================================================================
-	ILauncherDeviceGroup.h: Declares the ILauncherDeviceGroup interface.
-=============================================================================*/
-
 #pragma once
 
-/**
- * Type definition for shared pointers to instances of ILauncherDeviceGroup.
- */
+
+/** Type definition for shared pointers to instances of ILauncherDeviceGroup. */
 typedef TSharedPtr<class ILauncherDeviceGroup> ILauncherDeviceGroupPtr;
 
-/**
- * Type definition for shared references to instances of ILauncherDeviceGroup.
- */
+/** Type definition for shared references to instances of ILauncherDeviceGroup. */
 typedef TSharedRef<class ILauncherDeviceGroup> ILauncherDeviceGroupRef;
 
 
@@ -90,7 +83,6 @@ public:
 	 */
 	virtual void SetName( const FString& NewName ) = 0;
 
-
 public:
 
 	/**
@@ -107,11 +99,10 @@ public:
 	 */
 	virtual FOnLauncherDeviceGroupDeviceRemoved& OnDeviceRemoved( ) = 0;
 
-
 public:
 
 	/**
 	 * Virtual destructor.
 	 */
 	virtual ~ILauncherDeviceGroup( ) { }
-};
\ No newline at end of file
+};
diff --git a/Engine/Source/Developer/LauncherServices/Public/Interfaces/ILauncherProfile.h b/Engine/Source/Developer/LauncherServices/Public/Interfaces/ILauncherProfile.h
--- a/Engine/Source/Developer/LauncherServices/Public/Interfaces/ILauncherProfile.h
+++ b/Engine/Source/Developer/LauncherServices/Public/Interfaces/ILauncherProfile.h
@@ -1,9 +1,5 @@
 // Copyright 1998-2014 Epic Games, Inc. All Rights Reserved.
 
-/*=============================================================================
-	ILauncherProfile.h: Declares the ILauncherProfile interface.
-=============================================================================*/
-
 #pragma once
 
 
@@ -14,19 +10,13 @@ namespace ELauncherProfileCookModes
 	 */
 	enum Type
 	{
-		/**
-		 * Do not cook the build (default).
-		 */
+		/** Do not cook the build (default). */
 		DoNotCook,
 
-		/**
-		 * Pre-cook using user specified settings.
-		 */
+		/** Pre-cook using user specified settings. */
 		ByTheBook,
 
-		/**
-		 * Cook the build on the fly while the game is running.
-		 */
+		/** Cook the build on the fly while the game is running. */
 		OnTheFly
 	};
 }
@@ -39,19 +29,13 @@ namespace ELauncherProfileCookedMaps
 	 */
 	enum Type
 	{
-		/**
-		 * Cook all maps.
-		 */
+		/** Cook all maps. */
 		AllMaps,
 
-		/**
-		 * Don't cook any maps. Only startup packages will be cooked.
-		 */
+		/** Don't cook any maps. Only startup packages will be cooked. */
 		NoMaps,
 
-		/**
-		 * Cook user selected maps.
-		 */
+		/** Cook user selected maps. */
 		SelectedMaps
 
 	};
@@ -65,24 +49,16 @@ namespace ELauncherProfileDeploymentModes
 	 */
 	enum Type
 	{
-		/**
-		 * Do not deploy the build to any device.
-		 */
+		/** Do not deploy the build to any device. */
 		DoNotDeploy,
 
-		/**
-		 * Copy all required file to the device.
-		 */
+		/** Copy all required file to the device. */
 		CopyToDevice,
 
-		/**
-		 * Let the device get required files from a file server.
-		 */
+		/** Let the device get required files from a file server. */
 		FileServer,
 
-		/**
-		 * Copy a build from a repository to the device.
-		 */
+		/** Copy a build from a repository to the device. */
 		 CopyRepository,
 	};
 }
@@ -95,19 +71,13 @@ namespace ELauncherProfileLaunchModes
 	 */
 	enum Type
 	{
-		/**
-		 * Do not launch.
-		 */
+		/** Do not launch. */
 		DoNotLaunch,
 
-		/**
-		 * Launch with customized roles per device.
-		 */
+		/** Launch with customized roles per device. */
 		CustomRoles,
 
-		/**
-		 * Launch with the default role on all deployed devices.
-		 */
+		/** Launch with the default role on all deployed devices. */
 		DefaultRole	
 	};
 }
@@ -120,19 +90,13 @@ namespace ELauncherProfilePackagingModes
 	 */
 	enum Type
 	{
-		/**
-		 * Do not package.
-		 */
+		/** Do not package. */
 		DoNotPackage,
 
-		/**
-		 * Package and store the build locally.
-		 */
+		/** Package and store the build locally. */
 		Locally,
 
-		/**
-		 * Package and store the build in a shared repository.
-		 */
+		/** Package and store the build in a shared repository. */
 		SharedRepository
 	};
 }
@@ -151,77 +115,49 @@ namespace ELauncherProfileValidationErrors
 		 */
 		CopyToDeviceRequiresCookByTheBook,
 
-		/**
-		 * Custom launch roles are not yet supported.
-		 */
+		/** Custom launch roles are not yet supported. */
 		CustomRolesNotSupportedYet,
 
-		/**
-		 * A device group must be selected when deploying builds.
-		 */
+		/** A device group must be selected when deploying builds. */
 		DeployedDeviceGroupRequired,
 		
-		/**
-		 * The initial culture configured for launch is not part of the selected build.
-		 */
+		/** The initial culture configured for launch is not part of the selected build. */
 		InitialCultureNotAvailable,
 
-		/**
-		 * The initial map configured for launch is not part of the selected build.
-		 */
+		/** The initial map configured for launch is not part of the selected build. */
 		InitialMapNotAvailable,
 
-		/**
-		 * The specified launch command line is not formatted correctly.
-		 */
+		/** The specified launch command line is not formatted correctly. */
 		MalformedLaunchCommandLine,
 
-		/**
-		 * A build configuration is required when creating new builds.
-		 */
+		/** A build configuration is required when creating new builds. */
 		NoBuildConfigurationSelected,
 
-		/**
-		 * When cooking a build, at least one culture must be included.
-		 */
+		/** When cooking a build, at least one culture must be included. */
 		NoCookedCulturesSelected,
 
-		/**
-		 * One or more launch roles do not have a device assigned.
-		 */
+		/** One or more launch roles do not have a device assigned. */
 		NoLaunchRoleDeviceAssigned,
 
-		/**
-		 * At least one platform is required when creating new builds.
-		 */
+		/** At least one platform is required when creating new builds. */
 		NoPlatformSelected,
 
-		/**
-		 * A game is required when creating new builds.
-		 */
+		/** A game is required when creating new builds. */
 		NoProjectSelected,
 
-		/**
-		 * The deployment requires a package directory to be specified
-		 */
+		/** The deployment requires a package directory to be specified */
 		NoPackageDirectorySpecified,
 		
-		/**
-		 * The platform SDK is not installed but is required.
-		 */
+		/** The platform SDK is not installed but is required. */
 		NoPlatformSDKInstalled,
 	};
 }
 
 
-/**
- * Type definition for shared pointers to instances of ILauncherProfile.
- */
+/** Type definition for shared pointers to instances of ILauncherProfile. */
 typedef TSharedPtr<class ILauncherProfile> ILauncherProfilePtr;
 
-/**
- * Type definition for shared references to instances of ILauncherProfile.
- */
+/** Type definition for shared references to instances of ILauncherProfile. */
 typedef TSharedRef<class ILauncherProfile> ILauncherProfileRef;
 
 
@@ -232,11 +168,10 @@ typedef TSharedRef<class ILauncherProfile> ILauncherProfileRef;
  */
 DECLARE_MULTICAST_DELEGATE_OneParam(FOnLauncherProfileDeployedDeviceGroupChanged, const ILauncherDeviceGroupPtr&)
 
-/**
- * Delegate type for a change in project
- */
+/** Delegate type for a change in project */
 DECLARE_MULTICAST_DELEGATE(FOnProfileProjectChanged);
 
+
 /**
  * Interface for launcher profile.
  */
@@ -307,7 +242,6 @@ public:
 	 */
 	virtual void SetName( const FString& NewName ) = 0;
 
-
 public:
 
 	/**
@@ -543,7 +477,6 @@ public:
 	 */
 	virtual bool SupportsEngineMaps( ) const = 0;
 
-
 public:
 
 	/**
diff --git a/Engine/Source/Developer/LauncherServices/Public/Interfaces/ILauncherProfileLaunchRole.h b/Engine/Source/Developer/LauncherServices/Public/Interfaces/ILauncherProfileLaunchRole.h
--- a/Engine/Source/Developer/LauncherServices/Public/Interfaces/ILauncherProfileLaunchRole.h
+++ b/Engine/Source/Developer/LauncherServices/Public/Interfaces/ILauncherProfileLaunchRole.h
@@ -1,9 +1,5 @@
 // Copyright 1998-2014 Epic Games, Inc. All Rights Reserved.
 
-/*=============================================================================
-	ILauncherLaunchRole.h: Declares the ILauncherLaunchRole interface.
-=============================================================================*/
-
 #pragma once
 
 
@@ -14,24 +10,16 @@ namespace ELauncherProfileRoleInstanceTypes
 	 */
 	enum Type
 	{
-		/**
-		 * The instance is a dedicated server.
-		 */
+		/** The instance is a dedicated server. */
 		DedicatedServer,
 
-		/**
-		 * The instance is a listen server.
-		 */
+		/** The instance is a listen server. */
 		ListenServer,
 
-		/**
-		 * The instance is a game client.
-		 */
+		/** The instance is a game client. */
 		StandaloneClient,
 
-		/**
-		 * The instance is an Unreal Editor.
-		 */
+		/** The instance is an Unreal Editor. */
 		UnrealEditor
 	};
 
@@ -65,14 +53,10 @@ namespace ELauncherProfileRoleInstanceTypes
 }
 
 
-/**
- * Type definition for shared pointers to instances of ILauncherProfileLaunchRole.
- */
+/** Type definition for shared pointers to instances of ILauncherProfileLaunchRole. */
 typedef TSharedPtr<class ILauncherProfileLaunchRole> ILauncherProfileLaunchRolePtr;
 
-/**
- * Type definition for shared references to instances of ILauncherProfileLaunchRole.
- */
+/** Type definition for shared references to instances of ILauncherProfileLaunchRole. */
 typedef TSharedRef<class ILauncherProfileLaunchRole> ILauncherProfileLaunchRoleRef;
 
 
@@ -199,11 +183,10 @@ public:
 	 */
 	virtual void SetVsyncEnabled( bool Enabled ) = 0;
 
-
 public:
 
 	/**
 	 * Virtual destructor.
 	 */
 	virtual ~ILauncherProfileLaunchRole( ) { }
-};
\ No newline at end of file
+};
diff --git a/Engine/Source/Developer/LauncherServices/Public/Interfaces/ILauncherProfileManager.h b/Engine/Source/Developer/LauncherServices/Public/Interfaces/ILauncherProfileManager.h
--- a/Engine/Source/Developer/LauncherServices/Public/Interfaces/ILauncherProfileManager.h
+++ b/Engine/Source/Developer/LauncherServices/Public/Interfaces/ILauncherProfileManager.h
@@ -1,20 +1,12 @@
 // Copyright 1998-2014 Epic Games, Inc. All Rights Reserved.
 
-/*=============================================================================
-	ILauncherProfileManager.h: Declares the ILauncherProfileManager interface.
-=============================================================================*/
-
 #pragma once
 
 
-/**
- * Type definition for shared pointers to instances of ILauncherProfileManager.
- */
+/** Type definition for shared pointers to instances of ILauncherProfileManager. */
 typedef TSharedPtr<class ILauncherProfileManager> ILauncherProfileManagerPtr;
 
-/**
- * Type definition for shared references to instances of ILauncherProfileManager.
- */
+/** Type definition for shared references to instances of ILauncherProfileManager. */
 typedef TSharedRef<class ILauncherProfileManager> ILauncherProfileManagerRef;
 
 
@@ -91,7 +83,6 @@ public:
 	 */
 	virtual void RemoveDeviceGroup( const ILauncherDeviceGroupRef& DeviceGroup ) = 0;
 
-
 public:
 
 	/**
@@ -174,7 +165,6 @@ public:
 	 */
 	virtual void SaveProfile( const ILauncherProfileRef& Profile, FArchive& Archive ) = 0;
 
-
 public:
 
 	/**
@@ -195,7 +185,6 @@ public:
 	 */
 	virtual void SaveSettings( ) = 0;
 
-
 public:
 
 	/**
@@ -226,7 +215,6 @@ public:
 	 */
 	virtual FOnLauncherProfileManagerProfileRemoved& OnProfileRemoved( ) = 0;
 
-
 public:
 
 	/**
diff --git a/Engine/Source/Developer/LauncherServices/Public/Interfaces/ILauncherServicesModule.h b/Engine/Source/Developer/LauncherServices/Public/Interfaces/ILauncherServicesModule.h
--- a/Engine/Source/Developer/LauncherServices/Public/Interfaces/ILauncherServicesModule.h
+++ b/Engine/Source/Developer/LauncherServices/Public/Interfaces/ILauncherServicesModule.h
@@ -1,9 +1,5 @@
 // Copyright 1998-2014 Epic Games, Inc. All Rights Reserved.
 
-/*=============================================================================
-	ILauncherServicesModule.h: Declares the ILauncherServicesModule interface.
-=============================================================================*/
-
 #pragma once
 
 
diff --git a/Engine/Source/Developer/LauncherServices/Public/Interfaces/ILauncherTask.h b/Engine/Source/Developer/LauncherServices/Public/Interfaces/ILauncherTask.h
--- a/Engine/Source/Developer/LauncherServices/Public/Interfaces/ILauncherTask.h
+++ b/Engine/Source/Developer/LauncherServices/Public/Interfaces/ILauncherTask.h
@@ -14,34 +14,22 @@ namespace ELauncherTaskStatus
 	 */
 	enum Type
 	{
-		/**
-		 * The task is currently busy executing.
-		 */
+		/** The task is currently busy executing. */
 		Busy,
 
-		/**
-		 * The task has been canceled.
-		 */
+		/** The task has been canceled. */
 		Canceled,
 
-		/**
-		 * The task is being canceled.
-		 */
+		/** The task is being canceled. */
 		Canceling,
 
-		/**
-		 * The task completed successfully.
-		 */
+		/** The task completed successfully. */
 		Completed,
 
-		/**
-		 * The task failed.
-		 */
+		/** The task failed. */
 		Failed,
 
-		/**
-		 * The task is waiting to execute.
-		 */
+		/** The task is waiting to execute. */
 		Pending
 	};
 }
@@ -52,14 +40,11 @@ DECLARE_MULTICAST_DELEGATE_OneParam(FOnTaskStartedDelegate, const FString&);
 /** Delegate used to notify when a stage ends */
 DECLARE_MULTICAST_DELEGATE_OneParam(FOnTaskCompletedDelegate, const FString&);
 
-/**
- * Type definition for shared pointers to instances of ILauncherTask.
- */
+
+/** Type definition for shared pointers to instances of ILauncherTask. */
 typedef TSharedPtr<class ILauncherTask> ILauncherTaskPtr;
 
-/**
- * Type definition for shared references to instances of ILauncherTask.
- */
+/** Type definition for shared references to instances of ILauncherTask. */
 typedef TSharedRef<class ILauncherTask> ILauncherTaskRef;
 
 
@@ -125,4 +110,4 @@ public:
 	 * @return the delegate
 	 */
 	virtual FOnTaskCompletedDelegate& OnCompleted() = 0;
-};
\ No newline at end of file
+};
diff --git a/Engine/Source/Developer/LauncherServices/Public/Interfaces/ILauncherWorker.h b/Engine/Source/Developer/LauncherServices/Public/Interfaces/ILauncherWorker.h
--- a/Engine/Source/Developer/LauncherServices/Public/Interfaces/ILauncherWorker.h
+++ b/Engine/Source/Developer/LauncherServices/Public/Interfaces/ILauncherWorker.h
@@ -1,9 +1,5 @@
 // Copyright 1998-2014 Epic Games, Inc. All Rights Reserved.
 
-/*=============================================================================
-	ILauncherWorker.h: Declares the ILauncherWorker interface.
-=============================================================================*/
-
 #pragma once
 
 
@@ -22,16 +18,13 @@ namespace ELauncherWorkerStatus
 }
 
 
-/**
- * Type definition for shared pointers to instances of ILauncherWorker.
- */
+/** Type definition for shared pointers to instances of ILauncherWorker. */
 typedef TSharedPtr<class ILauncherWorker> ILauncherWorkerPtr;
 
-/**
- * Type definition for shared references to instances of ILauncherWorker.
- */
+/** Type definition for shared references to instances of ILauncherWorker. */
 typedef TSharedRef<class ILauncherWorker> ILauncherWorkerRef;
 
+
 /** Delegate used to notify of an output message */
 DECLARE_MULTICAST_DELEGATE_OneParam(FOutputMessageReceivedDelegate, const FString&);
 
@@ -46,6 +39,8 @@ DECLARE_MULTICAST_DELEGATE_OneParam(FOnLaunchCompletedDelegate, bool);
 
 /** Delegate used to notify when the launch was canceled */
 DECLARE_MULTICAST_DELEGATE(FOnLaunchCanceledDelegate);
+
+
 /**
  * Interface for launcher worker threads.
  */
@@ -115,4 +110,4 @@ protected:
 	 * Virtual destructor.
 	 */
 	virtual ~ILauncherWorker( ) { }
-};
\ No newline at end of file
+};
diff --git a/Engine/Source/Developer/LauncherServices/Public/LauncherServices.h b/Engine/Source/Developer/LauncherServices/Public/LauncherServices.h
--- a/Engine/Source/Developer/LauncherServices/Public/LauncherServices.h
+++ b/Engine/Source/Developer/LauncherServices/Public/LauncherServices.h
@@ -1,9 +1,5 @@
 // Copyright 1998-2014 Epic Games, Inc. All Rights Reserved.
 
-/*=============================================================================
-	LauncherServices.h: LauncherServices module public header file.
-=============================================================================*/
-
 #pragma once
 
 
@@ -31,4 +27,4 @@
 /* Common
  *****************************************************************************/
 
-#include "GameProjectHelper.h"
\ No newline at end of file
+#include "GameProjectHelper.h"
diff --git a/Engine/Source/Developer/Profiler/Private/ProfilerPrivatePCH.h b/Engine/Source/Developer/Profiler/Private/ProfilerPrivatePCH.h
--- a/Engine/Source/Developer/Profiler/Private/ProfilerPrivatePCH.h
+++ b/Engine/Source/Developer/Profiler/Private/ProfilerPrivatePCH.h
@@ -1,30 +1,21 @@
 // Copyright 1998-2014 Epic Games, Inc. All Rights Reserved.
 
-/*=============================================================================
-	ProfilerPrivatePCH.h: Pre-compiled header file for the Profiler module.
-=============================================================================*/
+#pragma once
 
-#ifndef __ProfilerPrivatePCH_h__
-#define __ProfilerPrivatePCH_h__
+#include "Profiler.h"
 
-/*-----------------------------------------------------------------------------
-	Public included
------------------------------------------------------------------------------*/
 
-#include "../Public/Profiler.h"
-#include "StatsData.h"
-#include "StatsFile.h"
-
-/*-----------------------------------------------------------------------------
-	Dependencies
------------------------------------------------------------------------------*/
+/* Private dependencies
+ *****************************************************************************/
 
 #include "Messaging.h"
 #include "ProfilerClient.h"
+#include "StatsData.h"
+#include "StatsFile.h"
+
 
-/*-----------------------------------------------------------------------------
-	Private includes
------------------------------------------------------------------------------*/
+/* Private includes
+ *****************************************************************************/
 
 // TODO: Remove dependencies between generic profiler code and profiler ui code
 
@@ -51,5 +42,3 @@
 #include "SProfilerWindow.h"
 #include "SHistogram.h"
 #include "StatDragDropOp.h"
-
-#endif // __ProfilerPrivatePCH_h__
diff --git a/Engine/Source/Developer/ProfilerClient/Private/ProfilerClientPrivatePCH.h b/Engine/Source/Developer/ProfilerClient/Private/ProfilerClientPrivatePCH.h
--- a/Engine/Source/Developer/ProfilerClient/Private/ProfilerClientPrivatePCH.h
+++ b/Engine/Source/Developer/ProfilerClient/Private/ProfilerClientPrivatePCH.h
@@ -1,16 +1,11 @@
 // Copyright 1998-2014 Epic Games, Inc. All Rights Reserved.
 
-/*=============================================================================
-	ProfilerClientPrivatePCH.h: Pre-compiled header file for the ProfilerClient module.
-=============================================================================*/
-
 #pragma once
 
-
-#include "../Public/ProfilerClient.h"
+#include "ProfilerClient.h"
 
 
-/* Dependencies
+/* Private dependencies
  *****************************************************************************/
 
 #include "Messaging.h"
diff --git a/Engine/Source/Developer/ProfilerMessages/Private/ProfilerMessagesPrivatePCH.h b/Engine/Source/Developer/ProfilerMessages/Private/ProfilerMessagesPrivatePCH.h
--- a/Engine/Source/Developer/ProfilerMessages/Private/ProfilerMessagesPrivatePCH.h
+++ b/Engine/Source/Developer/ProfilerMessages/Private/ProfilerMessagesPrivatePCH.h
@@ -1,24 +1,11 @@
 // Copyright 1998-2014 Epic Games, Inc. All Rights Reserved.
 
-/*=============================================================================
-	ProfilerMessagesPrivatePCH.h: Pre-compiled header file for the ProfilerMessages module.
-=============================================================================*/
+#pragma once
 
-#ifndef PROFILERMESSAGES_PRIVATEPCH_H
-#define PROFILERMESSAGES_PRIVATEPCH_H
+#include "ProfilerMessages.h"
 
 
-#include "../Public/ProfilerMessages.h"
-
-
-/* Dependencies
+/* Private dependencies
  *****************************************************************************/
 
 #include "ModuleManager.h"
-
-
-/* Private includes
- *****************************************************************************/
-
-
-#endif
