UE4 Refactoring. Changed OVERRIDE and FINAL macros to keywords override and final.

diff --git a/Engine/Source/Developer/LauncherServices/Private/Launcher/LauncherPackageCommands.h b/Engine/Source/Developer/LauncherServices/Private/Launcher/LauncherPackageCommands.h
--- a/Engine/Source/Developer/LauncherServices/Private/Launcher/LauncherPackageCommands.h
+++ b/Engine/Source/Developer/LauncherServices/Private/Launcher/LauncherPackageCommands.h
@@ -17,17 +17,17 @@ public:
 	{
 	}
 
-	virtual FString GetName() const OVERRIDE
+	virtual FString GetName() const override
 	{
 		return NSLOCTEXT("FLauncherTask", "LauncherStageTaskName", "Staging content").ToString();
 	}
 
-	virtual FString GetDesc() const OVERRIDE
+	virtual FString GetDesc() const override
 	{
 		return NSLOCTEXT("FLauncherTask", "LauncherStageTaskDesc", "Staging content for ").ToString() + TargetPlatform.PlatformName();
 	}
 
-	virtual FString GetArguments(FLauncherTaskChainState& ChainState) const OVERRIDE
+	virtual FString GetArguments(FLauncherTaskChainState& ChainState) const override
 	{
 		// build UAT command line parameters
 		FString CommandLine;
@@ -93,17 +93,17 @@ public:
 	{
 	}
 
-	virtual FString GetName() const OVERRIDE
+	virtual FString GetName() const override
 	{
 		return NSLOCTEXT("FLauncherTask", "LauncherStageTaskName", "Staging content").ToString();
 	}
 
-	virtual FString GetDesc() const OVERRIDE
+	virtual FString GetDesc() const override
 	{
 		return NSLOCTEXT("FLauncherTask", "LauncherStageTaskDesc", "Staging content for ").ToString() + TargetPlatform.PlatformName();
 	}
 
-	virtual FString GetArguments(FLauncherTaskChainState& ChainState) const OVERRIDE
+	virtual FString GetArguments(FLauncherTaskChainState& ChainState) const override
 	{
 		FString CommandLine;
 		FString InitialMap = ChainState.Profile->GetDefaultLaunchRole()->GetInitialMap();
@@ -182,17 +182,17 @@ public:
 	{
 	}
 
-	virtual FString GetName() const OVERRIDE
+	virtual FString GetName() const override
 	{
 		return NSLOCTEXT("FLauncherTask", "LauncherPackageTaskName", "Packaging content").ToString();
 	}
 
-	virtual FString GetDesc() const OVERRIDE
+	virtual FString GetDesc() const override
 	{
 		return NSLOCTEXT("FLauncherTask", "LauncherPackageTaskDesc", "Packaging content for ").ToString() + TargetPlatform.PlatformName();
 	}
 
-	virtual FString GetArguments(FLauncherTaskChainState& ChainState) const OVERRIDE
+	virtual FString GetArguments(FLauncherTaskChainState& ChainState) const override
 	{
 		// build UAT command line parameters
 		FString CommandLine;
@@ -229,14 +229,14 @@ public:
 		return CommandLine;
 	}
 
-	virtual bool PreExecute(FLauncherTaskChainState& ChainState) OVERRIDE
+	virtual bool PreExecute(FLauncherTaskChainState& ChainState) override
 	{
 		// disable the device check
 		const_cast<ITargetPlatform&>(TargetPlatform).EnableDeviceCheck(false);
 		return true;
 	}
 
-	virtual bool PostExecute(FLauncherTaskChainState& ChainState) OVERRIDE
+	virtual bool PostExecute(FLauncherTaskChainState& ChainState) override
 	{
 		// disable the device check
 		const_cast<ITargetPlatform&>(TargetPlatform).EnableDeviceCheck(true);
@@ -270,17 +270,17 @@ public:
 	{
 	}
 
-	virtual FString GetName() const OVERRIDE
+	virtual FString GetName() const override
 	{
 		return NSLOCTEXT("FLauncherTask", "LauncherPackageTaskName", "Packaging content").ToString();
 	}
 
-	virtual FString GetDesc() const OVERRIDE
+	virtual FString GetDesc() const override
 	{
 		return NSLOCTEXT("FLauncherTask", "LauncherPackageTaskDesc", "Packaging content for ").ToString() + TargetPlatform.PlatformName();
 	}
 
-	virtual FString GetArguments(FLauncherTaskChainState& ChainState) const OVERRIDE
+	virtual FString GetArguments(FLauncherTaskChainState& ChainState) const override
 	{
 		// build UAT command line parameters
 		FString CommandLine;
diff --git a/Engine/Source/Developer/LauncherServices/Private/Launcher/LauncherTask.h b/Engine/Source/Developer/LauncherServices/Private/Launcher/LauncherTask.h
--- a/Engine/Source/Developer/LauncherServices/Private/Launcher/LauncherTask.h
+++ b/Engine/Source/Developer/LauncherServices/Private/Launcher/LauncherTask.h
@@ -117,12 +117,12 @@ public:
 
 	// Begin FRunnable interface
 
-	virtual bool Init( ) OVERRIDE
+	virtual bool Init( ) override
 	{
 		return true;
 	}
 
-	virtual uint32 Run( ) OVERRIDE
+	virtual uint32 Run( ) override
 	{
 		Status = ELauncherTaskStatus::Busy;
 
@@ -159,12 +159,12 @@ public:
 		return 0;
 	}
 
-	virtual void Stop( ) OVERRIDE
+	virtual void Stop( ) override
 	{
 		Cancel();
 	}
 
-	virtual void Exit( ) OVERRIDE { }
+	virtual void Exit( ) override { }
 
 	// End FRunnable interface
 
@@ -172,7 +172,7 @@ public:
 
 	// Begin ILauncherTask interface
 
-	virtual void Cancel( ) OVERRIDE
+	virtual void Cancel( ) override
 	{
 		if (Status == ELauncherTaskStatus::Busy)
 		{
@@ -186,7 +186,7 @@ public:
 		CancelContinuations();
 	}
 
-	virtual FTimespan GetDuration( ) const OVERRIDE
+	virtual FTimespan GetDuration( ) const override
 	{
 		if (Status == ELauncherTaskStatus::Pending)
 		{
@@ -201,34 +201,34 @@ public:
 		return (EndTime - StartTime);
 	}
 
-	virtual const FString& GetName( ) const OVERRIDE
+	virtual const FString& GetName( ) const override
 	{
 		return Name;
 	}
 
-	virtual const FString& GetDesc( ) const OVERRIDE
+	virtual const FString& GetDesc( ) const override
 	{
 		return Desc;
 	}
 
-	virtual ELauncherTaskStatus::Type GetStatus( ) const OVERRIDE
+	virtual ELauncherTaskStatus::Type GetStatus( ) const override
 	{
 		return Status;
 	}
 
-	virtual bool IsFinished( ) const OVERRIDE
+	virtual bool IsFinished( ) const override
 	{
 		return ((Status == ELauncherTaskStatus::Canceled) ||
 				(Status == ELauncherTaskStatus::Completed) ||
 				(Status == ELauncherTaskStatus::Failed));
 	}
 
-	virtual FOnTaskStartedDelegate& OnStarted() OVERRIDE
+	virtual FOnTaskStartedDelegate& OnStarted() override
 	{
 		return TaskStarted;
 	}
 
-	virtual FOnTaskCompletedDelegate& OnCompleted() OVERRIDE
+	virtual FOnTaskCompletedDelegate& OnCompleted() override
 	{
 		return TaskCompleted;
 	}
diff --git a/Engine/Source/Developer/LauncherServices/Private/Launcher/LauncherUATTask.h b/Engine/Source/Developer/LauncherServices/Private/Launcher/LauncherUATTask.h
--- a/Engine/Source/Developer/LauncherServices/Private/Launcher/LauncherUATTask.h
+++ b/Engine/Source/Developer/LauncherServices/Private/Launcher/LauncherUATTask.h
@@ -39,7 +39,7 @@ protected:
 	 *
 	 * @return true if the task completed successfully, false otherwise.
 	 */
-	virtual bool PerformTask( FLauncherTaskChainState& ChainState ) OVERRIDE
+	virtual bool PerformTask( FLauncherTaskChainState& ChainState ) override
 	{
 		// spawn a UAT process to cook the data
 		// UAT executable
diff --git a/Engine/Source/Developer/LauncherServices/Private/Launcher/LauncherVerifyProfileTask.h b/Engine/Source/Developer/LauncherServices/Private/Launcher/LauncherVerifyProfileTask.h
--- a/Engine/Source/Developer/LauncherServices/Private/Launcher/LauncherVerifyProfileTask.h
+++ b/Engine/Source/Developer/LauncherServices/Private/Launcher/LauncherVerifyProfileTask.h
@@ -24,7 +24,7 @@ public:
 
 protected:
 
-	virtual bool PerformTask( FLauncherTaskChainState& ChainState ) OVERRIDE
+	virtual bool PerformTask( FLauncherTaskChainState& ChainState ) override
 	{
 		if (ChainState.Profile.IsValid())
 		{
diff --git a/Engine/Source/Developer/LauncherServices/Private/Launcher/LauncherWorker.h b/Engine/Source/Developer/LauncherServices/Private/Launcher/LauncherWorker.h
--- a/Engine/Source/Developer/LauncherServices/Private/Launcher/LauncherWorker.h
+++ b/Engine/Source/Developer/LauncherServices/Private/Launcher/LauncherWorker.h
@@ -28,24 +28,24 @@ public:
 
 public:
 
-	virtual bool Init( ) OVERRIDE;
+	virtual bool Init( ) override;
 
-	virtual uint32 Run( ) OVERRIDE;
+	virtual uint32 Run( ) override;
 
-	virtual void Stop( ) OVERRIDE;
+	virtual void Stop( ) override;
 
-	virtual void Exit( ) OVERRIDE { }
+	virtual void Exit( ) override { }
 
 public:
 
-	virtual void Cancel( ) OVERRIDE;
+	virtual void Cancel( ) override;
 
-	virtual ELauncherWorkerStatus::Type GetStatus( ) const  OVERRIDE
+	virtual ELauncherWorkerStatus::Type GetStatus( ) const  override
 	{
 		return Status;
 	}
 
-	virtual int32 GetTasks( TArray<ILauncherTaskPtr>& OutTasks ) const OVERRIDE;
+	virtual int32 GetTasks( TArray<ILauncherTaskPtr>& OutTasks ) const override;
 
 	virtual FOutputMessageReceivedDelegate& OnOutputReceived()
 	{
diff --git a/Engine/Source/Developer/LauncherServices/Private/LauncherServicesModule.cpp b/Engine/Source/Developer/LauncherServices/Private/LauncherServicesModule.cpp
--- a/Engine/Source/Developer/LauncherServices/Private/LauncherServicesModule.cpp
+++ b/Engine/Source/Developer/LauncherServices/Private/LauncherServicesModule.cpp
@@ -13,27 +13,27 @@ public:
 
 	// Begin ILauncherServicesModule interface
 
-	virtual ILauncherDeviceGroupRef CreateDeviceGroup( ) OVERRIDE
+	virtual ILauncherDeviceGroupRef CreateDeviceGroup( ) override
 	{
 		return MakeShareable(new FLauncherDeviceGroup());
 	}
 
-	virtual ILauncherDeviceGroupRef CreateDeviceGroup(const FGuid& Guid, const FString& Name ) OVERRIDE
+	virtual ILauncherDeviceGroupRef CreateDeviceGroup(const FGuid& Guid, const FString& Name ) override
 	{
 		return MakeShareable(new FLauncherDeviceGroup(Guid, Name));
 	}
 
-	virtual ILauncherRef CreateLauncher( ) OVERRIDE
+	virtual ILauncherRef CreateLauncher( ) override
 	{
 		return MakeShareable(new FLauncher());
 	}
 
-	virtual ILauncherProfileRef CreateProfile( const FString& ProfileName ) OVERRIDE
+	virtual ILauncherProfileRef CreateProfile( const FString& ProfileName ) override
 	{
 		return MakeShareable(new FLauncherProfile(ProfileName));
 	}
 
-	virtual ILauncherProfileManagerRef GetProfileManager( ) OVERRIDE
+	virtual ILauncherProfileManagerRef GetProfileManager( ) override
 	{
 		if (!ProfileManagerSingleton.IsValid())
 		{
@@ -44,11 +44,11 @@ public:
 	}
 	
 	DECLARE_DERIVED_EVENT(FLauncherServicesModule, ILauncherServicesModule::FLauncherServicesSDKNotInstalled, FLauncherServicesSDKNotInstalled);
-	virtual FLauncherServicesSDKNotInstalled& OnLauncherServicesSDKNotInstalled( ) OVERRIDE
+	virtual FLauncherServicesSDKNotInstalled& OnLauncherServicesSDKNotInstalled( ) override
 	{
 		return LauncherServicesSDKNotInstalled;
 	}
-	void BroadcastLauncherServicesSDKNotInstalled(const FString& PlatformName, const FString& DocLink) OVERRIDE
+	void BroadcastLauncherServicesSDKNotInstalled(const FString& PlatformName, const FString& DocLink) override
 	{
 		return LauncherServicesSDKNotInstalled.Broadcast(PlatformName, DocLink);
 	}
diff --git a/Engine/Source/Developer/LauncherServices/Private/Profiles/LauncherDeviceGroup.h b/Engine/Source/Developer/LauncherServices/Private/Profiles/LauncherDeviceGroup.h
--- a/Engine/Source/Developer/LauncherServices/Private/Profiles/LauncherDeviceGroup.h
+++ b/Engine/Source/Developer/LauncherServices/Private/Profiles/LauncherDeviceGroup.h
@@ -32,51 +32,51 @@ public:
 
 	// Begin ILauncherDeviceGroup interface
 
-	virtual void AddDevice( const FString& Device ) OVERRIDE
+	virtual void AddDevice( const FString& Device ) override
 	{
 		Devices.AddUnique(Device);
 
 		DeviceAddedDelegate.Broadcast(AsShared(), Device);
 	}
 
-	virtual const TArray<FString>& GetDevices( ) OVERRIDE
+	virtual const TArray<FString>& GetDevices( ) override
 	{
 		return Devices;
 	}
 
-	virtual FGuid GetId( ) const OVERRIDE
+	virtual FGuid GetId( ) const override
 	{
 		return Id;
 	}
 
-	virtual const FString& GetName( ) const OVERRIDE
+	virtual const FString& GetName( ) const override
 	{
 		return Name;
 	}
 
-	virtual int32 GetNumDevices( ) const OVERRIDE
+	virtual int32 GetNumDevices( ) const override
 	{
 		return Devices.Num();
 	}
 
-	virtual FOnLauncherDeviceGroupDeviceAdded& OnDeviceAdded( ) OVERRIDE
+	virtual FOnLauncherDeviceGroupDeviceAdded& OnDeviceAdded( ) override
 	{
 		return DeviceAddedDelegate;
 	}
 
-	virtual FOnLauncherDeviceGroupDeviceRemoved& OnDeviceRemoved( ) OVERRIDE
+	virtual FOnLauncherDeviceGroupDeviceRemoved& OnDeviceRemoved( ) override
 	{
 		return DeviceRemovedDelegate;
 	}
 
-	virtual void RemoveDevice( const FString& Device ) OVERRIDE
+	virtual void RemoveDevice( const FString& Device ) override
 	{
 		Devices.Remove(Device);
 
 		DeviceRemovedDelegate.Broadcast(AsShared(), Device);
 	}
 
-	virtual void SetName( const FString& NewName ) OVERRIDE
+	virtual void SetName( const FString& NewName ) override
 	{
 		Name = NewName;
 	}
diff --git a/Engine/Source/Developer/LauncherServices/Private/Profiles/LauncherProfile.h b/Engine/Source/Developer/LauncherServices/Private/Profiles/LauncherProfile.h
--- a/Engine/Source/Developer/LauncherServices/Private/Profiles/LauncherProfile.h
+++ b/Engine/Source/Developer/LauncherServices/Private/Profiles/LauncherProfile.h
@@ -56,28 +56,28 @@ public:
 
 	// Begin ILauncherProfile interface
 
-	virtual void AddCookedCulture( const FString& CultureName ) OVERRIDE
+	virtual void AddCookedCulture( const FString& CultureName ) override
 	{
 		CookedCultures.AddUnique(CultureName);
 
 		Validate();
 	}
 
-	virtual void AddCookedMap( const FString& MapName ) OVERRIDE
+	virtual void AddCookedMap( const FString& MapName ) override
 	{
 		CookedMaps.AddUnique(MapName);
 
 		Validate();
 	}
 
-	virtual void AddCookedPlatform( const FString& PlatformName ) OVERRIDE
+	virtual void AddCookedPlatform( const FString& PlatformName ) override
 	{
 		CookedPlatforms.AddUnique(PlatformName);
 
 		Validate();
 	}
 
-	virtual void ClearCookedCultures( ) OVERRIDE
+	virtual void ClearCookedCultures( ) override
 	{
 		if (CookedCultures.Num() > 0)
 		{
@@ -87,7 +87,7 @@ public:
 		}
 	}
 
-	virtual void ClearCookedMaps( ) OVERRIDE
+	virtual void ClearCookedMaps( ) override
 	{
 		if (CookedMaps.Num() > 0)
 		{
@@ -97,7 +97,7 @@ public:
 		}
 	}
 
-	virtual void ClearCookedPlatforms( ) OVERRIDE
+	virtual void ClearCookedPlatforms( ) override
 	{
 		if (CookedPlatforms.Num() > 0)
 		{
@@ -107,7 +107,7 @@ public:
 		}
 	}
 
-	virtual ILauncherProfileLaunchRolePtr CreateLaunchRole( ) OVERRIDE
+	virtual ILauncherProfileLaunchRolePtr CreateLaunchRole( ) override
 	{
 		ILauncherProfileLaunchRolePtr Role = MakeShareable(new FLauncherProfileLaunchRole());
 			
@@ -118,77 +118,77 @@ public:
 		return Role;
 	}
 
-	virtual EBuildConfigurations::Type GetBuildConfiguration( ) const OVERRIDE
+	virtual EBuildConfigurations::Type GetBuildConfiguration( ) const override
 	{
 		return BuildConfiguration;
 	}
 
-	virtual EBuildConfigurations::Type GetCookConfiguration( ) const OVERRIDE
+	virtual EBuildConfigurations::Type GetCookConfiguration( ) const override
 	{
 		return CookConfiguration;
 	}
 
-	virtual ELauncherProfileCookModes::Type GetCookMode( ) const OVERRIDE
+	virtual ELauncherProfileCookModes::Type GetCookMode( ) const override
 	{
 		return CookMode;
 	}
 
-	virtual const FString& GetCookOptions( ) const OVERRIDE
+	virtual const FString& GetCookOptions( ) const override
 	{
 		return CookOptions;
 	}
 
-	virtual const TArray<FString>& GetCookedCultures( ) const OVERRIDE
+	virtual const TArray<FString>& GetCookedCultures( ) const override
 	{
 		return CookedCultures;
 	}
 
-	virtual const TArray<FString>& GetCookedMaps( ) const OVERRIDE
+	virtual const TArray<FString>& GetCookedMaps( ) const override
 	{
 		return CookedMaps;
 	}
 
-	virtual const TArray<FString>& GetCookedPlatforms( ) const OVERRIDE
+	virtual const TArray<FString>& GetCookedPlatforms( ) const override
 	{
 		return CookedPlatforms;
 	}
 
-	virtual const ILauncherProfileLaunchRoleRef& GetDefaultLaunchRole( ) const OVERRIDE
+	virtual const ILauncherProfileLaunchRoleRef& GetDefaultLaunchRole( ) const override
 	{
 		return DefaultLaunchRole;
 	}
 
-	virtual ILauncherDeviceGroupPtr GetDeployedDeviceGroup( ) const OVERRIDE
+	virtual ILauncherDeviceGroupPtr GetDeployedDeviceGroup( ) const override
 	{
 		return DeployedDeviceGroup;
 	}
 
-	virtual ELauncherProfileDeploymentModes::Type GetDeploymentMode( ) const OVERRIDE
+	virtual ELauncherProfileDeploymentModes::Type GetDeploymentMode( ) const override
 	{
 		return DeploymentMode;
 	}
 
-    virtual bool GetForceClose() const OVERRIDE
+    virtual bool GetForceClose() const override
     {
         return ForceClose;
     }
     
-	virtual FGuid GetId( ) const OVERRIDE
+	virtual FGuid GetId( ) const override
 	{
 		return Id;
 	}
 
-	virtual ELauncherProfileLaunchModes::Type GetLaunchMode( ) const OVERRIDE
+	virtual ELauncherProfileLaunchModes::Type GetLaunchMode( ) const override
 	{
 		return LaunchMode;
 	}
 
-	virtual const TArray<ILauncherProfileLaunchRolePtr>& GetLaunchRoles( ) const OVERRIDE
+	virtual const TArray<ILauncherProfileLaunchRolePtr>& GetLaunchRoles( ) const override
 	{
 		return LaunchRoles;
 	}
 
-	virtual const int32 GetLaunchRolesFor( const FString& DeviceId, TArray<ILauncherProfileLaunchRolePtr>& OutRoles ) OVERRIDE
+	virtual const int32 GetLaunchRolesFor( const FString& DeviceId, TArray<ILauncherProfileLaunchRolePtr>& OutRoles ) override
 	{
 		OutRoles.Empty();
 
@@ -212,22 +212,22 @@ public:
 		return OutRoles.Num();
 	}
 
-	virtual FString GetName( ) const OVERRIDE
+	virtual FString GetName( ) const override
 	{
 		return Name;
 	}
 
-	virtual ELauncherProfilePackagingModes::Type GetPackagingMode( ) const OVERRIDE
+	virtual ELauncherProfilePackagingModes::Type GetPackagingMode( ) const override
 	{
 		return PackagingMode;
 	}
 
-	virtual FString GetPackageDirectory( ) const OVERRIDE
+	virtual FString GetPackageDirectory( ) const override
 	{
 		return PackageDir;
 	}
 
-	virtual FString GetProjectName( ) const OVERRIDE
+	virtual FString GetProjectName( ) const override
 	{
 		if (!ProjectPath.IsEmpty())
 		{
@@ -237,7 +237,7 @@ public:
 		return FString();
 	}
 
-	virtual FString GetProjectBasePath() const OVERRIDE
+	virtual FString GetProjectBasePath() const override
 	{
 		if (!ProjectPath.IsEmpty())
 		{
@@ -249,37 +249,37 @@ public:
 		return FString();
 	}
 
-	virtual const FString& GetProjectPath( ) const OVERRIDE
+	virtual const FString& GetProjectPath( ) const override
 	{
 		return ProjectPath;
 	}
 
-    virtual uint32 GetTimeout() const OVERRIDE
+    virtual uint32 GetTimeout() const override
     {
         return Timeout;
     }
     
-	virtual bool HasValidationError( ELauncherProfileValidationErrors::Type Error ) const OVERRIDE
+	virtual bool HasValidationError( ELauncherProfileValidationErrors::Type Error ) const override
 	{
 		return ValidationErrors.Contains(Error);
 	}
 
-	virtual bool IsBuilding() const OVERRIDE
+	virtual bool IsBuilding() const override
 	{
 		return BuildGame;
 	}
 
-	virtual bool IsCookingIncrementally( ) const OVERRIDE
+	virtual bool IsCookingIncrementally( ) const override
 	{
 		return CookIncremental;
 	}
 
-	virtual bool IsCookingUnversioned( ) const OVERRIDE
+	virtual bool IsCookingUnversioned( ) const override
 	{
 		return CookUnversioned;
 	}
 
-	virtual bool IsDeployablePlatform( const FString& PlatformName ) OVERRIDE
+	virtual bool IsDeployablePlatform( const FString& PlatformName ) override
 	{
 		if (CookMode == ELauncherProfileCookModes::ByTheBook)
 		{
@@ -289,55 +289,55 @@ public:
 		return true;
 	}
 
-	virtual bool IsFileServerHidden( ) const OVERRIDE
+	virtual bool IsFileServerHidden( ) const override
 	{
 		return HideFileServerWindow;
 	}
 
-	virtual bool IsFileServerStreaming( ) const OVERRIDE
+	virtual bool IsFileServerStreaming( ) const override
 	{
 		return DeployStreamingServer;
 	}
 
-	virtual bool IsPackingWithUnrealPak( ) const  OVERRIDE
+	virtual bool IsPackingWithUnrealPak( ) const  override
 	{
 		return DeployWithUnrealPak;
 	}
 
-	virtual bool IsValidForLaunch( ) OVERRIDE
+	virtual bool IsValidForLaunch( ) override
 	{
 		return (ValidationErrors.Num() == 0);
 	}
 
-	virtual void RemoveCookedCulture( const FString& CultureName ) OVERRIDE
+	virtual void RemoveCookedCulture( const FString& CultureName ) override
 	{
 		CookedCultures.Remove(CultureName);
 
 		Validate();
 	}
 
-	virtual void RemoveCookedMap( const FString& MapName ) OVERRIDE
+	virtual void RemoveCookedMap( const FString& MapName ) override
 	{
 		CookedMaps.Remove(MapName);
 
 		Validate();
 	}
 
-	virtual void RemoveCookedPlatform( const FString& PlatformName ) OVERRIDE
+	virtual void RemoveCookedPlatform( const FString& PlatformName ) override
 	{
 		CookedPlatforms.Remove(PlatformName);
 
 		Validate();
 	}
 
-	virtual void RemoveLaunchRole( const ILauncherProfileLaunchRoleRef& Role ) OVERRIDE
+	virtual void RemoveLaunchRole( const ILauncherProfileLaunchRoleRef& Role ) override
 	{
 		LaunchRoles.Remove(Role);
 
 		Validate();
 	}
 
-	virtual bool Serialize( FArchive& Archive ) OVERRIDE
+	virtual bool Serialize( FArchive& Archive ) override
 	{
 		int32 Version = LAUNCHERSERVICES_PROFILEVERSION;
 
@@ -415,7 +415,7 @@ public:
 		return true;
 	}
 
-	virtual void SetDefaults( ) OVERRIDE
+	virtual void SetDefaults( ) override
 	{
 		// default project settings
 		if (FPaths::IsProjectFilePathSet())
@@ -479,7 +479,7 @@ public:
 		Validate();
 	}
 
-	virtual void SetBuildGame(bool Build) OVERRIDE
+	virtual void SetBuildGame(bool Build) override
 	{
 		if (BuildGame != Build)
 		{
@@ -489,7 +489,7 @@ public:
 		}
 	}
 
-	virtual void SetBuildConfiguration( EBuildConfigurations::Type Configuration ) OVERRIDE
+	virtual void SetBuildConfiguration( EBuildConfigurations::Type Configuration ) override
 	{
 		if (BuildConfiguration != Configuration)
 		{
@@ -499,7 +499,7 @@ public:
 		}
 	}
 
-	virtual void SetCookConfiguration( EBuildConfigurations::Type Configuration ) OVERRIDE
+	virtual void SetCookConfiguration( EBuildConfigurations::Type Configuration ) override
 	{
 		if (CookConfiguration != Configuration)
 		{
@@ -509,7 +509,7 @@ public:
 		}
 	}
 
-	virtual void SetCookMode( ELauncherProfileCookModes::Type Mode ) OVERRIDE
+	virtual void SetCookMode( ELauncherProfileCookModes::Type Mode ) override
 	{
 		if (CookMode != Mode)
 		{
@@ -519,7 +519,7 @@ public:
 		}
 	}
 
-	virtual void SetDeployWithUnrealPak( bool UseUnrealPak ) OVERRIDE
+	virtual void SetDeployWithUnrealPak( bool UseUnrealPak ) override
 	{
 		if (DeployWithUnrealPak != UseUnrealPak)
 		{
@@ -529,7 +529,7 @@ public:
 		}
 	}
 
-	virtual void SetDeployedDeviceGroup( const ILauncherDeviceGroupPtr& DeviceGroup ) OVERRIDE
+	virtual void SetDeployedDeviceGroup( const ILauncherDeviceGroupPtr& DeviceGroup ) override
 	{
 		if(DeployedDeviceGroup.IsValid())
 		{
@@ -551,7 +551,7 @@ public:
 		Validate();
 	}
 
-	virtual void SetDeploymentMode( ELauncherProfileDeploymentModes::Type Mode ) OVERRIDE
+	virtual void SetDeploymentMode( ELauncherProfileDeploymentModes::Type Mode ) override
 	{
 		if (DeploymentMode != Mode)
 		{
@@ -561,7 +561,7 @@ public:
 		}
 	}
 
-    virtual void SetForceClose( bool Close ) OVERRIDE
+    virtual void SetForceClose( bool Close ) override
     {
         if (ForceClose != Close)
         {
@@ -570,12 +570,12 @@ public:
         }
     }
     
-	virtual void SetHideFileServerWindow( bool Hide ) OVERRIDE
+	virtual void SetHideFileServerWindow( bool Hide ) override
 	{
 		HideFileServerWindow = Hide;
 	}
 
-	virtual void SetIncrementalCooking( bool Incremental ) OVERRIDE
+	virtual void SetIncrementalCooking( bool Incremental ) override
 	{
 		if (CookIncremental != Incremental)
 		{
@@ -585,7 +585,7 @@ public:
 		}
 	}
 
-	virtual void SetLaunchMode( ELauncherProfileLaunchModes::Type Mode ) OVERRIDE
+	virtual void SetLaunchMode( ELauncherProfileLaunchModes::Type Mode ) override
 	{
 		if (LaunchMode != Mode)
 		{
@@ -595,7 +595,7 @@ public:
 		}
 	}
 
-	virtual void SetName( const FString& NewName ) OVERRIDE
+	virtual void SetName( const FString& NewName ) override
 	{
 		if (Name != NewName)
 		{
@@ -605,7 +605,7 @@ public:
 		}
 	}
 
-	virtual void SetPackagingMode( ELauncherProfilePackagingModes::Type Mode ) OVERRIDE
+	virtual void SetPackagingMode( ELauncherProfilePackagingModes::Type Mode ) override
 	{
 		if (PackagingMode != Mode)
 		{
@@ -615,7 +615,7 @@ public:
 		}
 	}
 
-	virtual void SetPackageDirectory( const FString& Dir ) OVERRIDE
+	virtual void SetPackageDirectory( const FString& Dir ) override
 	{
 		if (PackageDir != Dir)
 		{
@@ -625,7 +625,7 @@ public:
 		}
 	}
 
-	virtual void SetProjectPath( const FString& Path ) OVERRIDE
+	virtual void SetProjectPath( const FString& Path ) override
 	{
 		if (ProjectPath != Path)
 		{
@@ -638,7 +638,7 @@ public:
 		}
 	}
 
-	virtual void SetStreamingFileServer( bool Streaming ) OVERRIDE
+	virtual void SetStreamingFileServer( bool Streaming ) override
 	{
 		if (DeployStreamingServer != Streaming)
 		{
@@ -648,7 +648,7 @@ public:
 		}
 	}
 
-    virtual void SetTimeout( uint32 InTime ) OVERRIDE
+    virtual void SetTimeout( uint32 InTime ) override
     {
         if (Timeout != InTime)
         {
@@ -658,7 +658,7 @@ public:
         }
     }
     
-	virtual void SetUnversionedCooking( bool Unversioned ) OVERRIDE
+	virtual void SetUnversionedCooking( bool Unversioned ) override
 	{
 		if (CookUnversioned != Unversioned)
 		{
@@ -668,12 +668,12 @@ public:
 		}
 	}
 
-	virtual bool SupportsEngineMaps( ) const OVERRIDE
+	virtual bool SupportsEngineMaps( ) const override
 	{
 		return false;
 	}
 
-	virtual FOnProfileProjectChanged& OnProjectChanged() OVERRIDE
+	virtual FOnProfileProjectChanged& OnProjectChanged() override
 	{
 		return ProjectChangedDelegate;
 	}
diff --git a/Engine/Source/Developer/LauncherServices/Private/Profiles/LauncherProfileLaunchRole.h b/Engine/Source/Developer/LauncherServices/Private/Profiles/LauncherProfileLaunchRole.h
--- a/Engine/Source/Developer/LauncherServices/Private/Profiles/LauncherProfileLaunchRole.h
+++ b/Engine/Source/Developer/LauncherServices/Private/Profiles/LauncherProfileLaunchRole.h
@@ -28,42 +28,42 @@ public:
 
 	// Begin ILauncherProfileLaunchRole interface
 
-	virtual const FString& GetAssignedDevice( ) const OVERRIDE
+	virtual const FString& GetAssignedDevice( ) const override
 	{
 		return AssignedDevice;
 	}
 
-	virtual const FString& GetCommandLine( ) const OVERRIDE
+	virtual const FString& GetCommandLine( ) const override
 	{
 		return CommandLine;
 	}
 
-	virtual const FString& GetInitialCulture( ) const OVERRIDE
+	virtual const FString& GetInitialCulture( ) const override
 	{
 		return InitialCulture;
 	}
 
-	virtual const FString& GetInitialMap( ) const OVERRIDE
+	virtual const FString& GetInitialMap( ) const override
 	{
 		return InitialMapName;
 	}
 
-	virtual ELauncherProfileRoleInstanceTypes::Type GetInstanceType( ) const OVERRIDE
+	virtual ELauncherProfileRoleInstanceTypes::Type GetInstanceType( ) const override
 	{
 		return InstanceType;
 	}
 
-	virtual const FString& GetName( ) const OVERRIDE
+	virtual const FString& GetName( ) const override
 	{
 		return Name;
 	}
 
-	virtual bool IsVsyncEnabled( ) const OVERRIDE
+	virtual bool IsVsyncEnabled( ) const override
 	{
 		return VsyncEnabled;
 	}
 
-	virtual void Serialize( FArchive& Archive ) OVERRIDE
+	virtual void Serialize( FArchive& Archive ) override
 	{
 		Archive << AssignedDevice
 				<< CommandLine
@@ -75,32 +75,32 @@ public:
 				<< VsyncEnabled;
 	}
 
-	virtual void SetCommandLine( const FString& NewCommandLine ) OVERRIDE
+	virtual void SetCommandLine( const FString& NewCommandLine ) override
 	{
 		CommandLine = NewCommandLine;
 	}
 
-	virtual void SetInitialCulture( const FString& CultureName ) OVERRIDE
+	virtual void SetInitialCulture( const FString& CultureName ) override
 	{
 		InitialCulture = CultureName;
 	}
 
-	virtual void SetInitialMap( const FString& MapName ) OVERRIDE
+	virtual void SetInitialMap( const FString& MapName ) override
 	{
 		InitialMapName = MapName;
 	}
 
-	virtual void SetInstanceType( ELauncherProfileRoleInstanceTypes::Type InInstanceType ) OVERRIDE
+	virtual void SetInstanceType( ELauncherProfileRoleInstanceTypes::Type InInstanceType ) override
 	{
 		InstanceType = InInstanceType;
 	}
 
-	virtual void SetName( const FString& NewName ) OVERRIDE
+	virtual void SetName( const FString& NewName ) override
 	{
 		Name = NewName;
 	}
 
-	virtual void SetVsyncEnabled( bool Enabled ) OVERRIDE
+	virtual void SetVsyncEnabled( bool Enabled ) override
 	{
 		VsyncEnabled = Enabled;
 	}
diff --git a/Engine/Source/Developer/LauncherServices/Private/Profiles/LauncherProfileManager.h b/Engine/Source/Developer/LauncherServices/Private/Profiles/LauncherProfileManager.h
--- a/Engine/Source/Developer/LauncherServices/Private/Profiles/LauncherProfileManager.h
+++ b/Engine/Source/Developer/LauncherServices/Private/Profiles/LauncherProfileManager.h
@@ -19,55 +19,55 @@ public:
 
 	// Begin ILauncherProfileManager interface
 
-	virtual void AddDeviceGroup( const ILauncherDeviceGroupRef& DeviceGroup ) OVERRIDE;
+	virtual void AddDeviceGroup( const ILauncherDeviceGroupRef& DeviceGroup ) override;
 
-	virtual ILauncherDeviceGroupRef AddNewDeviceGroup( ) OVERRIDE;
+	virtual ILauncherDeviceGroupRef AddNewDeviceGroup( ) override;
 
-	virtual ILauncherProfileRef AddNewProfile( ) OVERRIDE;
+	virtual ILauncherProfileRef AddNewProfile( ) override;
 
-	virtual void AddProfile( const ILauncherProfileRef& Profile ) OVERRIDE;
+	virtual void AddProfile( const ILauncherProfileRef& Profile ) override;
 
-	virtual ILauncherProfilePtr FindProfile( const FString& ProfileName ) OVERRIDE;
+	virtual ILauncherProfilePtr FindProfile( const FString& ProfileName ) override;
 
-	virtual const TArray<ILauncherDeviceGroupPtr>& GetAllDeviceGroups( ) const OVERRIDE;
+	virtual const TArray<ILauncherDeviceGroupPtr>& GetAllDeviceGroups( ) const override;
 
-	virtual const TArray<ILauncherProfilePtr>& GetAllProfiles( ) const OVERRIDE;
+	virtual const TArray<ILauncherProfilePtr>& GetAllProfiles( ) const override;
 
-	virtual ILauncherDeviceGroupPtr GetDeviceGroup( const FGuid& GroupId ) const OVERRIDE;
+	virtual ILauncherDeviceGroupPtr GetDeviceGroup( const FGuid& GroupId ) const override;
 
-	virtual ILauncherProfilePtr GetProfile( const FGuid& ProfileId ) const OVERRIDE;
+	virtual ILauncherProfilePtr GetProfile( const FGuid& ProfileId ) const override;
 
-	virtual ILauncherProfilePtr LoadProfile( FArchive& Archive ) OVERRIDE;
+	virtual ILauncherProfilePtr LoadProfile( FArchive& Archive ) override;
 
-	virtual void LoadSettings( ) OVERRIDE;
+	virtual void LoadSettings( ) override;
 
-	virtual FOnLauncherProfileManagerDeviceGroupAdded& OnDeviceGroupAdded( ) OVERRIDE
+	virtual FOnLauncherProfileManagerDeviceGroupAdded& OnDeviceGroupAdded( ) override
 	{
 		return DeviceGroupAddedDelegate;
 	}
 
-	virtual FOnLauncherProfileManagerDeviceGroupRemoved& OnDeviceGroupRemoved( ) OVERRIDE
+	virtual FOnLauncherProfileManagerDeviceGroupRemoved& OnDeviceGroupRemoved( ) override
 	{
 		return DeviceGroupRemovedDelegate;
 	}
 
-	virtual FOnLauncherProfileManagerProfileAdded& OnProfileAdded( ) OVERRIDE
+	virtual FOnLauncherProfileManagerProfileAdded& OnProfileAdded( ) override
 	{
 		return ProfileAddedDelegate;
 	}
 
-	virtual FOnLauncherProfileManagerProfileRemoved& OnProfileRemoved( ) OVERRIDE
+	virtual FOnLauncherProfileManagerProfileRemoved& OnProfileRemoved( ) override
 	{
 		return ProfileRemovedDelegate;
 	}
 
-	virtual void RemoveDeviceGroup( const ILauncherDeviceGroupRef& DeviceGroup ) OVERRIDE;
+	virtual void RemoveDeviceGroup( const ILauncherDeviceGroupRef& DeviceGroup ) override;
 
-	virtual void RemoveProfile( const ILauncherProfileRef& Profile ) OVERRIDE;
+	virtual void RemoveProfile( const ILauncherProfileRef& Profile ) override;
 
-	virtual void SaveProfile( const ILauncherProfileRef& Profile, FArchive& Archive ) OVERRIDE;
+	virtual void SaveProfile( const ILauncherProfileRef& Profile, FArchive& Archive ) override;
 
-	virtual void SaveSettings( ) OVERRIDE;
+	virtual void SaveSettings( ) override;
 
 	// End ILauncherProfileManager interface
 
diff --git a/Engine/Source/Developer/Linux/LinuxTargetPlatform/Private/LinuxTargetDevice.h b/Engine/Source/Developer/Linux/LinuxTargetPlatform/Private/LinuxTargetDevice.h
--- a/Engine/Source/Developer/Linux/LinuxTargetPlatform/Private/LinuxTargetDevice.h
+++ b/Engine/Source/Developer/Linux/LinuxTargetPlatform/Private/LinuxTargetDevice.h
@@ -40,39 +40,39 @@ public:
 
 public:
 
-	virtual bool Connect( ) OVERRIDE
+	virtual bool Connect( ) override
 	{
 		return true;
 	}
 
-	virtual bool Deploy( const FString& SourceFolder, FString& OutAppId ) OVERRIDE;
+	virtual bool Deploy( const FString& SourceFolder, FString& OutAppId ) override;
 
 	virtual void Disconnect( )
 	{ }
 
-	virtual ETargetDeviceTypes::Type GetDeviceType( ) const OVERRIDE
+	virtual ETargetDeviceTypes::Type GetDeviceType( ) const override
 	{
 		return ETargetDeviceTypes::Desktop;
 	}
 
-	virtual FTargetDeviceId GetId( ) const OVERRIDE
+	virtual FTargetDeviceId GetId( ) const override
 	{
 		return TargetDeviceId;
 	}
 
-	virtual FString GetName( ) const OVERRIDE
+	virtual FString GetName( ) const override
 	{
 		return DeviceName;
 	}
 
-	virtual FString GetOperatingSystemName( ) OVERRIDE
+	virtual FString GetOperatingSystemName( ) override
 	{
 		return TEXT("GNU/Linux");
 	}
 
-	virtual int32 GetProcessSnapshot( TArray<FTargetDeviceProcessInfo>& OutProcessInfos ) OVERRIDE;
+	virtual int32 GetProcessSnapshot( TArray<FTargetDeviceProcessInfo>& OutProcessInfos ) override;
 
-	virtual const class ITargetPlatform& GetTargetPlatform( ) const OVERRIDE
+	virtual const class ITargetPlatform& GetTargetPlatform( ) const override
 	{
 		return TargetPlatform;
 	}
@@ -82,36 +82,36 @@ public:
 		return true;
 	}
 
-	virtual bool IsDefault( ) const OVERRIDE
+	virtual bool IsDefault( ) const override
 	{
 		return true;
 	}
 
-	virtual bool PowerOff( bool Force ) OVERRIDE
+	virtual bool PowerOff( bool Force ) override
 	{
 		return false;
 	}
 
-	virtual bool PowerOn( ) OVERRIDE
+	virtual bool PowerOn( ) override
 	{
 		return false;
 	}
 
-	virtual bool Launch( const FString& AppId, EBuildConfigurations::Type BuildConfiguration, EBuildTargets::Type BuildTarget, const FString& Params, uint32* OutProcessId ) OVERRIDE;
+	virtual bool Launch( const FString& AppId, EBuildConfigurations::Type BuildConfiguration, EBuildTargets::Type BuildTarget, const FString& Params, uint32* OutProcessId ) override;
 	
-	virtual bool Reboot( bool bReconnect = false ) OVERRIDE;
+	virtual bool Reboot( bool bReconnect = false ) override;
 
-	virtual bool Run( const FString& ExecutablePath, const FString& Params, uint32* OutProcessId ) OVERRIDE;
+	virtual bool Run( const FString& ExecutablePath, const FString& Params, uint32* OutProcessId ) override;
 
-	virtual bool SupportsFeature( ETargetDeviceFeatures::Type Feature ) const OVERRIDE;
+	virtual bool SupportsFeature( ETargetDeviceFeatures::Type Feature ) const override;
 
-	virtual bool SupportsSdkVersion( const FString& VersionString ) const OVERRIDE;
+	virtual bool SupportsSdkVersion( const FString& VersionString ) const override;
 
-	virtual void SetUserCredentials( const FString & UserName, const FString & UserPassword ) OVERRIDE;
+	virtual void SetUserCredentials( const FString & UserName, const FString & UserPassword ) override;
 
-	virtual bool GetUserCredentials( FString & OutUserName, FString & OutUserPassword ) OVERRIDE;
+	virtual bool GetUserCredentials( FString & OutUserName, FString & OutUserPassword ) override;
 
-	virtual bool TerminateProcess( const int32 ProcessId ) OVERRIDE;
+	virtual bool TerminateProcess( const int32 ProcessId ) override;
 
 
 private:
diff --git a/Engine/Source/Developer/Linux/LinuxTargetPlatform/Private/LinuxTargetPlatform.h b/Engine/Source/Developer/Linux/LinuxTargetPlatform/Private/LinuxTargetPlatform.h
--- a/Engine/Source/Developer/Linux/LinuxTargetPlatform/Private/LinuxTargetPlatform.h
+++ b/Engine/Source/Developer/Linux/LinuxTargetPlatform/Private/LinuxTargetPlatform.h
@@ -43,9 +43,9 @@ public:
 
 	// Begin ITargetPlatform interface
 
-	virtual void EnableDeviceCheck(bool OnOff) OVERRIDE {}
+	virtual void EnableDeviceCheck(bool OnOff) override {}
 
-	virtual void GetAllDevices( TArray<ITargetDevicePtr>& OutDevices ) const OVERRIDE
+	virtual void GetAllDevices( TArray<ITargetDevicePtr>& OutDevices ) const override
 	{
 		// TODO: ping all the machines in a local segment and/or try to connect to port 22 of those that respond
 		OutDevices.Reset();
@@ -55,17 +55,17 @@ public:
 		}
 	}
 
-	virtual ECompressionFlags GetBaseCompressionMethod( ) const OVERRIDE
+	virtual ECompressionFlags GetBaseCompressionMethod( ) const override
 	{
 		return COMPRESS_ZLIB;
 	}
 
-	virtual bool GenerateStreamingInstallManifest(const TMultiMap<FString, int32>& ChunkMap, const TSet<int32>& ChunkIDsInUse) const OVERRIDE
+	virtual bool GenerateStreamingInstallManifest(const TMultiMap<FString, int32>& ChunkMap, const TSet<int32>& ChunkIDsInUse) const override
 	{
 		return true;
 	}
 
-	virtual ITargetDevicePtr GetDefaultDevice( ) const OVERRIDE
+	virtual ITargetDevicePtr GetDefaultDevice( ) const override
 	{
 		if (LocalDevice.IsValid())
 		{
@@ -75,7 +75,7 @@ public:
 		return nullptr;
 	}
 
-	virtual ITargetDevicePtr GetDevice( const FTargetDeviceId& DeviceId ) OVERRIDE
+	virtual ITargetDevicePtr GetDevice( const FTargetDeviceId& DeviceId ) override
 	{
 		if (LocalDevice.IsValid() && (DeviceId == LocalDevice->GetId()))
 		{
@@ -86,12 +86,12 @@ public:
 		return MakeShareable(new FLinuxTargetDevice(*this, UATFriendlyId, DeviceId.GetDeviceName()));
 	}
 
-	virtual bool IsRunningPlatform( ) const OVERRIDE
+	virtual bool IsRunningPlatform( ) const override
 	{
 		return false;
 	}
 
-	bool SupportsFeature(ETargetPlatformFeatures::Type Feature) const OVERRIDE
+	bool SupportsFeature(ETargetPlatformFeatures::Type Feature) const override
 	{
 		if (Feature == ETargetPlatformFeatures::UserCredentials || Feature == ETargetPlatformFeatures::Packaging)
 		{
@@ -102,7 +102,7 @@ public:
 	}
 
 #if WITH_ENGINE
-	virtual void GetAllPossibleShaderFormats( TArray<FName>& OutFormats ) const OVERRIDE
+	virtual void GetAllPossibleShaderFormats( TArray<FName>& OutFormats ) const override
 	{
 		// no shaders needed for dedicated server target
 		if (!IS_DEDICATED_SERVER)
@@ -115,17 +115,17 @@ public:
 		}
 	}
 
-	virtual void GetAllTargetedShaderFormats( TArray<FName>& OutFormats ) const OVERRIDE
+	virtual void GetAllTargetedShaderFormats( TArray<FName>& OutFormats ) const override
 	{
 		GetAllPossibleShaderFormats( OutFormats );
 	}
 
-	virtual const class FStaticMeshLODSettings& GetStaticMeshLODSettings( ) const OVERRIDE
+	virtual const class FStaticMeshLODSettings& GetStaticMeshLODSettings( ) const override
 	{
 		return StaticMeshLODSettings;
 	}
 
-	virtual void GetTextureFormats( const UTexture* InTexture, TArray<FName>& OutFormats ) const OVERRIDE
+	virtual void GetTextureFormats( const UTexture* InTexture, TArray<FName>& OutFormats ) const override
 	{
 		if (!IS_DEDICATED_SERVER)
 		{
@@ -134,12 +134,12 @@ public:
 		}
 	}
 
-	virtual const struct FTextureLODSettings& GetTextureLODSettings( ) const OVERRIDE
+	virtual const struct FTextureLODSettings& GetTextureLODSettings( ) const override
 	{
 		return TextureLODSettings;
 	}
 
-	virtual FName GetWaveFormat( class USoundWave* Wave ) const OVERRIDE
+	virtual FName GetWaveFormat( class USoundWave* Wave ) const override
 	{
 		static FName NAME_OGG(TEXT("OGG"));
 
@@ -148,13 +148,13 @@ public:
 #endif //WITH_ENGINE
 
 	DECLARE_DERIVED_EVENT(FAndroidTargetPlatform, ITargetPlatform::FOnTargetDeviceDiscovered, FOnTargetDeviceDiscovered);
-	virtual FOnTargetDeviceDiscovered& OnDeviceDiscovered( ) OVERRIDE
+	virtual FOnTargetDeviceDiscovered& OnDeviceDiscovered( ) override
 	{
 		return DeviceDiscoveredEvent;
 	}
 
 	DECLARE_DERIVED_EVENT(FAndroidTargetPlatform, ITargetPlatform::FOnTargetDeviceLost, FOnTargetDeviceLost);
-	virtual FOnTargetDeviceLost& OnDeviceLost( ) OVERRIDE
+	virtual FOnTargetDeviceLost& OnDeviceLost( ) override
 	{
 		return DeviceLostEvent;
 	}
diff --git a/Engine/Source/Developer/Linux/LinuxTargetPlatform/Private/LinuxTargetPlatformModule.cpp b/Engine/Source/Developer/Linux/LinuxTargetPlatform/Private/LinuxTargetPlatformModule.cpp
--- a/Engine/Source/Developer/Linux/LinuxTargetPlatform/Private/LinuxTargetPlatformModule.cpp
+++ b/Engine/Source/Developer/Linux/LinuxTargetPlatform/Private/LinuxTargetPlatformModule.cpp
@@ -37,7 +37,7 @@ public:
 	
 	// Begin ITargetPlatformModule interface
 
-	virtual ITargetPlatform* GetTargetPlatform( ) OVERRIDE
+	virtual ITargetPlatform* GetTargetPlatform( ) override
 	{
 		if (Singleton == NULL)
 		{
@@ -54,7 +54,7 @@ public:
 
 	// Begin IModuleInterface interface
 
-	virtual void StartupModule() OVERRIDE
+	virtual void StartupModule() override
 	{
 		TargetSettings = ConstructObject<ULinuxTargetSettings>(ULinuxTargetSettings::StaticClass(), GetTransientPackage(), "LinuxTargetSettings", RF_Standalone);
 		TargetSettings->AddToRoot();
@@ -71,7 +71,7 @@ public:
 		}
 	}
 
-	virtual void ShutdownModule() OVERRIDE
+	virtual void ShutdownModule() override
 	{
 		TargetSettings->RemoveFromRoot();
 	}
diff --git a/Engine/Source/Developer/LogVisualizer/Private/LogVisualizer.h b/Engine/Source/Developer/LogVisualizer/Private/LogVisualizer.h
--- a/Engine/Source/Developer/LogVisualizer/Private/LogVisualizer.h
+++ b/Engine/Source/Developer/LogVisualizer/Private/LogVisualizer.h
@@ -8,8 +8,8 @@
 class FLogVisualizer : public ILogVisualizer
 {
 public:
-	virtual void SummonUI(UWorld* InWorld) OVERRIDE;
-	virtual void CloseUI(UWorld* InWorld) OVERRIDE;
+	virtual void SummonUI(UWorld* InWorld) override;
+	virtual void CloseUI(UWorld* InWorld) override;
 
 	virtual bool IsRecording();
 	// End ILogVisualizer interface
diff --git a/Engine/Source/Developer/LogVisualizer/Private/SFilterList.cpp b/Engine/Source/Developer/LogVisualizer/Private/SFilterList.cpp
--- a/Engine/Source/Developer/LogVisualizer/Private/SFilterList.cpp
+++ b/Engine/Source/Developer/LogVisualizer/Private/SFilterList.cpp
@@ -21,7 +21,7 @@ public:
 		OnFilterMiddleButtonClicked = NewFilterMiddleButtonClicked;
 	}
 
-	virtual FReply OnMouseButtonDoubleClick(const FGeometry& InMyGeometry, const FPointerEvent& InMouseEvent) OVERRIDE
+	virtual FReply OnMouseButtonDoubleClick(const FGeometry& InMyGeometry, const FPointerEvent& InMouseEvent) override
 	{
 		if (InMouseEvent.GetEffectingButton() == EKeys::LeftMouseButton && OnFilterDoubleClicked.IsBound())
 		{
@@ -33,7 +33,7 @@ public:
 		}
 	}
 
-	virtual FReply OnMouseButtonUp(const FGeometry& InMyGeometry, const FPointerEvent& InMouseEvent) OVERRIDE
+	virtual FReply OnMouseButtonUp(const FGeometry& InMyGeometry, const FPointerEvent& InMouseEvent) override
 	{
 		if (InMouseEvent.GetEffectingButton() == EKeys::MiddleMouseButton && OnFilterMiddleButtonClicked.IsBound())
 		{
diff --git a/Engine/Source/Developer/LogVisualizer/Private/SLogBar.h b/Engine/Source/Developer/LogVisualizer/Private/SLogBar.h
--- a/Engine/Source/Developer/LogVisualizer/Private/SLogBar.h
+++ b/Engine/Source/Developer/LogVisualizer/Private/SLogBar.h
@@ -60,7 +60,7 @@ public:
 	 *
 	 * @return The maximum layer ID attained by this widget or any of its children.
 	 */
-	virtual int32 OnPaint(const FGeometry& AllottedGeometry, const FSlateRect& MyClippingRect, FSlateWindowElementList& OutDrawElements, int32 LayerId, const FWidgetStyle& InWidgetStyle, bool bParentEnabled) const OVERRIDE;
+	virtual int32 OnPaint(const FGeometry& AllottedGeometry, const FSlateRect& MyClippingRect, FSlateWindowElementList& OutDrawElements, int32 LayerId, const FWidgetStyle& InWidgetStyle, bool bParentEnabled) const override;
 	
 	/**
 	 * Ticks this widget.  Override in derived classes, but always call the parent implementation.
@@ -79,9 +79,9 @@ public:
 	 *
 	 * @return Whether the event was handled along with possible requests for the system to take action.
 	 */
-	virtual FReply OnMouseButtonUp(const FGeometry& MyGeometry, const FPointerEvent& MouseEvent) OVERRIDE;
+	virtual FReply OnMouseButtonUp(const FGeometry& MyGeometry, const FPointerEvent& MouseEvent) override;
 
-	virtual FReply OnMouseButtonDown(const FGeometry& MyGeometry, const FPointerEvent& MouseEvent) OVERRIDE;
+	virtual FReply OnMouseButtonDown(const FGeometry& MyGeometry, const FPointerEvent& MouseEvent) override;
 
 	/**
 	 * The system calls this method to notify the widget that a mouse moved within it. This event is bubbled.
@@ -91,7 +91,7 @@ public:
 	 *
 	 * @return Whether the event was handled along with possible requests for the system to take action.
 	 */
-	virtual FReply OnMouseMove(const FGeometry& MyGeometry, const FPointerEvent& MouseEvent) OVERRIDE;
+	virtual FReply OnMouseMove(const FGeometry& MyGeometry, const FPointerEvent& MouseEvent) override;
 	
 	void OnCurrentTimeChanged(float NewTime);
 
diff --git a/Engine/Source/Developer/LogVisualizer/Private/SLogVisualizer.h b/Engine/Source/Developer/LogVisualizer/Private/SLogVisualizer.h
--- a/Engine/Source/Developer/LogVisualizer/Private/SLogVisualizer.h
+++ b/Engine/Source/Developer/LogVisualizer/Private/SLogVisualizer.h
@@ -76,11 +76,11 @@ public:
 
 	virtual ~SLogVisualizer();
 
-	virtual void Tick( const FGeometry& AllottedGeometry, const double InCurrentTime, const float InDeltaTime ) OVERRIDE;
+	virtual void Tick( const FGeometry& AllottedGeometry, const double InCurrentTime, const float InDeltaTime ) override;
 
 	// OVERRIDES
-	virtual FReply OnMouseWheel( const FGeometry& MyGeometry, const FPointerEvent& MouseEvent ) OVERRIDE;
-	virtual FReply OnKeyDown( const FGeometry& MyGeometry, const FKeyboardEvent& InKeyboardEvent) OVERRIDE;
+	virtual FReply OnMouseWheel( const FGeometry& MyGeometry, const FPointerEvent& MouseEvent ) override;
+	virtual FReply OnKeyDown( const FGeometry& MyGeometry, const FKeyboardEvent& InKeyboardEvent) override;
 
 	// Get delegates
 	const FSlateBrush* GetRecordButtonBrush() const;
diff --git a/Engine/Source/Developer/Mac/MacClientTargetPlatform/Private/MacClientTargetPlatformModule.cpp b/Engine/Source/Developer/Mac/MacClientTargetPlatform/Private/MacClientTargetPlatformModule.cpp
--- a/Engine/Source/Developer/Mac/MacClientTargetPlatform/Private/MacClientTargetPlatformModule.cpp
+++ b/Engine/Source/Developer/Mac/MacClientTargetPlatform/Private/MacClientTargetPlatformModule.cpp
@@ -26,7 +26,7 @@ public:
 		Singleton = NULL;
 	}
 
-	virtual ITargetPlatform* GetTargetPlatform( ) OVERRIDE
+	virtual ITargetPlatform* GetTargetPlatform( ) override
 	{
 		if (Singleton == NULL)
 		{
diff --git a/Engine/Source/Developer/Mac/MacNoEditorTargetPlatform/Private/MacNoEditorTargetPlatformModule.cpp b/Engine/Source/Developer/Mac/MacNoEditorTargetPlatform/Private/MacNoEditorTargetPlatformModule.cpp
--- a/Engine/Source/Developer/Mac/MacNoEditorTargetPlatform/Private/MacNoEditorTargetPlatformModule.cpp
+++ b/Engine/Source/Developer/Mac/MacNoEditorTargetPlatform/Private/MacNoEditorTargetPlatformModule.cpp
@@ -26,7 +26,7 @@ public:
 		Singleton = NULL;
 	}
 
-	virtual ITargetPlatform* GetTargetPlatform( ) OVERRIDE
+	virtual ITargetPlatform* GetTargetPlatform( ) override
 	{
 		if (Singleton == NULL)
 		{
diff --git a/Engine/Source/Developer/Mac/MacTargetPlatform/Private/GenericMacTargetPlatform.h b/Engine/Source/Developer/Mac/MacTargetPlatform/Private/GenericMacTargetPlatform.h
--- a/Engine/Source/Developer/Mac/MacTargetPlatform/Private/GenericMacTargetPlatform.h
+++ b/Engine/Source/Developer/Mac/MacTargetPlatform/Private/GenericMacTargetPlatform.h
@@ -39,25 +39,25 @@ public:
 
 	// Begin ITargetPlatform interface
 
-	virtual void EnableDeviceCheck(bool OnOff) OVERRIDE {}
+	virtual void EnableDeviceCheck(bool OnOff) override {}
 
-	virtual void GetAllDevices( TArray<ITargetDevicePtr>& OutDevices ) const OVERRIDE
+	virtual void GetAllDevices( TArray<ITargetDevicePtr>& OutDevices ) const override
 	{
 		OutDevices.Reset();
 		OutDevices.Add(LocalDevice);
 	}
 
-	virtual ECompressionFlags GetBaseCompressionMethod( ) const OVERRIDE
+	virtual ECompressionFlags GetBaseCompressionMethod( ) const override
 	{
 		return COMPRESS_ZLIB;
 	}
 
-	virtual bool GenerateStreamingInstallManifest(const TMultiMap<FString, int32>& ChunkMap, const TSet<int32>& ChunkIDsInUse) const OVERRIDE
+	virtual bool GenerateStreamingInstallManifest(const TMultiMap<FString, int32>& ChunkMap, const TSet<int32>& ChunkIDsInUse) const override
 	{
 		return true;
 	}
 
-	virtual ITargetDevicePtr GetDefaultDevice( ) const OVERRIDE
+	virtual ITargetDevicePtr GetDefaultDevice( ) const override
 	{
 		return LocalDevice;
 	}
@@ -72,13 +72,13 @@ public:
 		return NULL;
 	}
 
-	virtual bool IsRunningPlatform( ) const OVERRIDE
+	virtual bool IsRunningPlatform( ) const override
 	{
 		// Must be Mac platform as editor for this to be considered a running platform
 		return PLATFORM_MAC && !UE_SERVER && !UE_GAME && WITH_EDITOR && HAS_EDITOR_DATA;
 	}
 
-	virtual bool SupportsFeature( ETargetPlatformFeatures::Type Feature ) const OVERRIDE
+	virtual bool SupportsFeature( ETargetPlatformFeatures::Type Feature ) const override
 	{
 		// we currently do not have a build target for MacServer
 		if (Feature == ETargetPlatformFeatures::Packaging)
@@ -90,7 +90,7 @@ return TSuper::SupportsFeature(Feature);
 	}
 
 #if WITH_ENGINE
-	virtual void GetAllPossibleShaderFormats( TArray<FName>& OutFormats ) const OVERRIDE
+	virtual void GetAllPossibleShaderFormats( TArray<FName>& OutFormats ) const override
 	{
 		// no shaders needed for dedicated server target
 		if (!IS_DEDICATED_SERVER)
@@ -100,19 +100,19 @@ return TSuper::SupportsFeature(Feature);
 		}
 	}
 
-	virtual void GetAllTargetedShaderFormats(TArray<FName>& OutFormats) const OVERRIDE
+	virtual void GetAllTargetedShaderFormats(TArray<FName>& OutFormats) const override
 	{
 		GetAllPossibleShaderFormats( OutFormats );
 	}
 
 
-	virtual const class FStaticMeshLODSettings& GetStaticMeshLODSettings( ) const OVERRIDE
+	virtual const class FStaticMeshLODSettings& GetStaticMeshLODSettings( ) const override
 	{
 		return StaticMeshLODSettings;
 	}
 
 
-	virtual void GetTextureFormats( const UTexture* InTexture, TArray<FName>& OutFormats ) const OVERRIDE
+	virtual void GetTextureFormats( const UTexture* InTexture, TArray<FName>& OutFormats ) const override
 	{
 		if (IS_DEDICATED_SERVER)
 		{
@@ -127,13 +127,13 @@ return TSuper::SupportsFeature(Feature);
 	}
 
 
-	virtual const struct FTextureLODSettings& GetTextureLODSettings( ) const OVERRIDE
+	virtual const struct FTextureLODSettings& GetTextureLODSettings( ) const override
 	{
 		return TextureLODSettings;
 	}
 
 
-	virtual FName GetWaveFormat( class USoundWave* Wave ) const OVERRIDE
+	virtual FName GetWaveFormat( class USoundWave* Wave ) const override
 	{
 		static FName NAME_OGG(TEXT("OGG"));
 
@@ -142,13 +142,13 @@ return TSuper::SupportsFeature(Feature);
 #endif //WITH_ENGINE
 
 	DECLARE_DERIVED_EVENT(TGenericMacTargetPlatform, ITargetPlatform::FOnTargetDeviceDiscovered, FOnTargetDeviceDiscovered);
-	virtual FOnTargetDeviceDiscovered& OnDeviceDiscovered( ) OVERRIDE
+	virtual FOnTargetDeviceDiscovered& OnDeviceDiscovered( ) override
 	{
 		return DeviceDiscoveredEvent;
 	}
 
 	DECLARE_DERIVED_EVENT(TGenericMacTargetPlatform, ITargetPlatform::FOnTargetDeviceLost, FOnTargetDeviceLost);
-	virtual FOnTargetDeviceLost& OnDeviceLost( ) OVERRIDE
+	virtual FOnTargetDeviceLost& OnDeviceLost( ) override
 	{
 		return DeviceLostEvent;
 	}
diff --git a/Engine/Source/Developer/Mac/MacTargetPlatform/Private/LocalMacTargetDevice.h b/Engine/Source/Developer/Mac/MacTargetPlatform/Private/LocalMacTargetDevice.h
--- a/Engine/Source/Developer/Mac/MacTargetPlatform/Private/LocalMacTargetDevice.h
+++ b/Engine/Source/Developer/Mac/MacTargetPlatform/Private/LocalMacTargetDevice.h
@@ -27,12 +27,12 @@ public:
 
 public:
 
-	virtual bool Connect( ) OVERRIDE
+	virtual bool Connect( ) override
 	{
 		return true;
 	}
 
-	virtual bool Deploy( const FString& SourceFolder, FString& OutAppId ) OVERRIDE
+	virtual bool Deploy( const FString& SourceFolder, FString& OutAppId ) override
 	{
 		OutAppId = TEXT("");
 
@@ -61,33 +61,33 @@ public:
 	virtual void Disconnect( )
 	{ }
 
-	virtual ETargetDeviceTypes::Type GetDeviceType( ) const OVERRIDE
+	virtual ETargetDeviceTypes::Type GetDeviceType( ) const override
 	{
 		return ETargetDeviceTypes::Desktop;
 	}
 
-	virtual FTargetDeviceId GetId( ) const OVERRIDE
+	virtual FTargetDeviceId GetId( ) const override
 	{
 		return FTargetDeviceId(TargetPlatform.PlatformName(), GetName());
 	}
 
-	virtual FString GetName( ) const OVERRIDE
+	virtual FString GetName( ) const override
 	{
 		return FPlatformProcess::ComputerName();
 	}
 
-	virtual FString GetOperatingSystemName( ) OVERRIDE
+	virtual FString GetOperatingSystemName( ) override
 	{
 		return TEXT("OS X");
 	}
 
-	virtual int32 GetProcessSnapshot( TArray<FTargetDeviceProcessInfo>& OutProcessInfos ) OVERRIDE
+	virtual int32 GetProcessSnapshot( TArray<FTargetDeviceProcessInfo>& OutProcessInfos ) override
 	{
 		// @todo Mac: implement process snapshots
 		return 0;
 	}
 
-	virtual const class ITargetPlatform& GetTargetPlatform( ) const OVERRIDE
+	virtual const class ITargetPlatform& GetTargetPlatform( ) const override
 	{
 		return TargetPlatform;
 	}
@@ -97,12 +97,12 @@ public:
 		return true;
 	}
 
-	virtual bool IsDefault( ) const OVERRIDE
+	virtual bool IsDefault( ) const override
 	{
 		return true;
 	}
 
-	virtual bool Launch( const FString& AppId, EBuildConfigurations::Type BuildConfiguration, EBuildTargets::Type Target, const FString& Params, uint32* OutProcessId ) OVERRIDE
+	virtual bool Launch( const FString& AppId, EBuildConfigurations::Type BuildConfiguration, EBuildTargets::Type Target, const FString& Params, uint32* OutProcessId ) override
 	{
 		// build executable path
 		FString PlatformName = TEXT("Mac");
@@ -119,17 +119,17 @@ public:
 		return ProcessHandle.Close();
 	}
 
-	virtual bool PowerOff( bool Force ) OVERRIDE
+	virtual bool PowerOff( bool Force ) override
 	{
 		return false;
 	}
 
-	virtual bool PowerOn( ) OVERRIDE
+	virtual bool PowerOn( ) override
 	{
 		return false;
 	}
 
-	virtual bool Reboot( bool bReconnect = false ) OVERRIDE
+	virtual bool Reboot( bool bReconnect = false ) override
 	{
 #if PLATFORM_MAC
 		NSAppleScript* Script = [[NSAppleScript alloc] initWithSource:@"tell application \"System Events\" to restart"];
@@ -139,13 +139,13 @@ public:
 		return true;
 	}
 
-	virtual bool Run( const FString& ExecutablePath, const FString& Params, uint32* OutProcessId ) OVERRIDE
+	virtual bool Run( const FString& ExecutablePath, const FString& Params, uint32* OutProcessId ) override
 	{
 		FProcHandle ProcessHandle = FPlatformProcess::CreateProc(*ExecutablePath, *Params, true, false, false, OutProcessId, 0, NULL, NULL);
 		return ProcessHandle.Close();
 	}
 
-	virtual bool SupportsFeature( ETargetDeviceFeatures::Type Feature ) const OVERRIDE
+	virtual bool SupportsFeature( ETargetDeviceFeatures::Type Feature ) const override
 	{
 		switch (Feature)
 		{
@@ -163,22 +163,22 @@ public:
 		return false;
 	}
 
-	virtual bool SupportsSdkVersion( const FString& VersionString ) const OVERRIDE
+	virtual bool SupportsSdkVersion( const FString& VersionString ) const override
 	{
 		// @todo filter SDK versions
 		return true;
 	}
 
-	virtual void SetUserCredentials( const FString & UserName, const FString & UserPassword ) OVERRIDE
+	virtual void SetUserCredentials( const FString & UserName, const FString & UserPassword ) override
 	{
 	}
 
-	virtual bool GetUserCredentials( FString & OutUserName, FString & OutUserPassword ) OVERRIDE
+	virtual bool GetUserCredentials( FString & OutUserName, FString & OutUserPassword ) override
 	{
 		return false;
 	}
 
-	virtual bool TerminateProcess( const int32 ProcessId ) OVERRIDE
+	virtual bool TerminateProcess( const int32 ProcessId ) override
 	{
 		return false;
 	}
diff --git a/Engine/Source/Developer/Mac/MacTargetPlatform/Private/MacTargetPlatformModule.cpp b/Engine/Source/Developer/Mac/MacTargetPlatform/Private/MacTargetPlatformModule.cpp
--- a/Engine/Source/Developer/Mac/MacTargetPlatform/Private/MacTargetPlatformModule.cpp
+++ b/Engine/Source/Developer/Mac/MacTargetPlatform/Private/MacTargetPlatformModule.cpp
@@ -37,7 +37,7 @@ public:
 
 	// Begin ITargetPlatformModule interface
 
-	virtual ITargetPlatform* GetTargetPlatform() OVERRIDE
+	virtual ITargetPlatform* GetTargetPlatform() override
 	{
 		if (Singleton == NULL)
 		{
@@ -54,7 +54,7 @@ public:
 
 	// Begin IModuleInterface interface
 
-	virtual void StartupModule() OVERRIDE
+	virtual void StartupModule() override
 	{
 		TargetSettings = ConstructObject<UMacTargetSettings>(UMacTargetSettings::StaticClass(), GetTransientPackage(), "MacTargetSettings", RF_Standalone);
 		TargetSettings->AddToRoot();
@@ -71,7 +71,7 @@ public:
 		}
 	}
 
-	virtual void ShutdownModule() OVERRIDE
+	virtual void ShutdownModule() override
 	{
 		ISettingsModule* SettingsModule = ISettingsModule::Get();
 
diff --git a/Engine/Source/Developer/MeshBoneReduction/Private/MeshBoneReduction.cpp b/Engine/Source/Developer/MeshBoneReduction/Private/MeshBoneReduction.cpp
--- a/Engine/Source/Developer/MeshBoneReduction/Private/MeshBoneReduction.cpp
+++ b/Engine/Source/Developer/MeshBoneReduction/Private/MeshBoneReduction.cpp
@@ -8,11 +8,11 @@ class FMeshBoneReductionModule : public IMeshBoneReductionModule
 {
 public:
 	// IModuleInterface interface.
-	virtual void StartupModule() OVERRIDE;
-	virtual void ShutdownModule() OVERRIDE;
+	virtual void StartupModule() override;
+	virtual void ShutdownModule() override;
 
 	// IMeshBoneReductionModule interface.
-	virtual class IMeshBoneReduction* GetMeshBoneReductionInterface() OVERRIDE;
+	virtual class IMeshBoneReduction* GetMeshBoneReductionInterface() override;
 };
 
 DEFINE_LOG_CATEGORY_STATIC(LogMeshBoneReduction, Log, All);
diff --git a/Engine/Source/Developer/MeshSimplifier/Private/QuadricMeshReduction.cpp b/Engine/Source/Developer/MeshSimplifier/Private/QuadricMeshReduction.cpp
--- a/Engine/Source/Developer/MeshSimplifier/Private/QuadricMeshReduction.cpp
+++ b/Engine/Source/Developer/MeshSimplifier/Private/QuadricMeshReduction.cpp
@@ -10,12 +10,12 @@ class FQuadricSimplifierMeshReductionModule : public IMeshReductionModule
 {
 public:
 	// IModuleInterface interface.
-	virtual void StartupModule() OVERRIDE;
-	virtual void ShutdownModule() OVERRIDE;
+	virtual void StartupModule() override;
+	virtual void ShutdownModule() override;
 
 	// IMeshReductionModule interface.
-	virtual class IMeshReduction* GetMeshReductionInterface() OVERRIDE;
-	virtual class IMeshMerging* GetMeshMergingInterface() OVERRIDE;
+	virtual class IMeshReduction* GetMeshReductionInterface() override;
+	virtual class IMeshMerging* GetMeshMergingInterface() override;
 };
 
 
@@ -110,7 +110,7 @@ public:
 class FQuadricSimplifierMeshReduction : public IMeshReduction
 {
 public:
-	virtual const FString& GetVersionString() const OVERRIDE
+	virtual const FString& GetVersionString() const override
 	{
 		static FString Version = TEXT("1.0");
 		return Version;
diff --git a/Engine/Source/Developer/MeshUtilities/Private/MeshUtilities.cpp b/Engine/Source/Developer/MeshUtilities/Private/MeshUtilities.cpp
--- a/Engine/Source/Developer/MeshUtilities/Private/MeshUtilities.cpp
+++ b/Engine/Source/Developer/MeshUtilities/Private/MeshUtilities.cpp
@@ -466,7 +466,7 @@ private:
 	/** True if NvTriStrip is being used for tri order optimization. */
 	bool bUsingNvTriStrip;
 	// IMeshUtilities interface.
-	virtual const FString& GetVersionString() const OVERRIDE
+	virtual const FString& GetVersionString() const override
 	{
 		return VersionString;
 	}
@@ -475,7 +475,7 @@ private:
 		TArray<FStaticMeshSourceModel>& SourceModels,
 		const TArray<UMaterialInterface*>& Materials,
 		const FStaticMeshLODGroup& LODGroup
-		) OVERRIDE;
+		) override;
 
 	virtual bool BuildSkeletalMesh( FStaticLODModel& LODModel, const FReferenceSkeleton& RefSkeleton, const TArray<FVertInfluence>& Influences, const TArray<FMeshWedge>& Wedges, const TArray<FMeshFace>& Faces, const TArray<FVector>& Points, const TArray<int32>& PointToOriginalMap, bool bKeepOverlappingVertices = false, bool bComputeNormals = true, bool bComputeTangents = true );
 
@@ -489,13 +489,13 @@ private:
 		uint32& MaxCharts,
 		float& MaxDesiredStretch,
 		FText& OutError
-		) OVERRIDE;
+		) override;
 
-	virtual bool LayoutUVs(FRawMesh& RawMesh, uint32 TextureResolution, uint32 TexCoordIndex, FText& OutError) OVERRIDE;
-	virtual IMeshReduction* GetMeshReductionInterface() OVERRIDE;
-	virtual IMeshMerging* GetMeshMergingInterface() OVERRIDE;
-	virtual void CacheOptimizeIndexBuffer(TArray<uint16>& Indices) OVERRIDE;
-	virtual void CacheOptimizeIndexBuffer(TArray<uint32>& Indices) OVERRIDE;
+	virtual bool LayoutUVs(FRawMesh& RawMesh, uint32 TextureResolution, uint32 TexCoordIndex, FText& OutError) override;
+	virtual IMeshReduction* GetMeshReductionInterface() override;
+	virtual IMeshMerging* GetMeshMergingInterface() override;
+	virtual void CacheOptimizeIndexBuffer(TArray<uint16>& Indices) override;
+	virtual void CacheOptimizeIndexBuffer(TArray<uint32>& Indices) override;
 	void CacheOptimizeVertexAndIndexBuffer(TArray<FStaticMeshBuildVertex>& Vertices,TArray<TArray<uint32> >& PerSectionIndices,TArray<int32>& WedgeMap);
 	
 	virtual void BuildSkeletalAdjacencyIndexBuffer(
@@ -503,9 +503,9 @@ private:
 		const uint32 TexCoordCount,
 		const TArray<uint32>& Indices,
 		TArray<uint32>& OutPnAenIndices
-		) OVERRIDE;
+		) override;
 
-	virtual void CalcBoneVertInfos(USkeletalMesh* SkeletalMesh, TArray<FBoneVertInfo>& Infos, bool bOnlyDominant) OVERRIDE;
+	virtual void CalcBoneVertInfos(USkeletalMesh* SkeletalMesh, TArray<FBoneVertInfo>& Infos, bool bOnlyDominant) override;
 
 	/**
 	 * Builds a renderable skeletal mesh LOD model. Note that the array of chunks
@@ -518,8 +518,8 @@ private:
 	void BuildSkeletalModelFromChunks(FStaticLODModel& LODModel,const FReferenceSkeleton& RefSkeleton,TArray<FSkinnedMeshChunk*>& Chunks,const TArray<int32>& PointToOriginalMap);
 
 	// IModuleInterface interface.
-	virtual void StartupModule() OVERRIDE;
-	virtual void ShutdownModule() OVERRIDE;
+	virtual void StartupModule() override;
+	virtual void ShutdownModule() override;
 
 	virtual void MergeActors(
 		const TArray<AActor*>& SourceActors,
@@ -527,7 +527,7 @@ private:
 		const FString& PackageName,
 		TArray<UObject*>& 
 		OutAssetsToSync, 
-		FVector& OutMergedActorLocation) const OVERRIDE;
+		FVector& OutMergedActorLocation) const override;
 	
 	virtual void CreateProxyMesh( 
 		const TArray<AActor*>& Actors, 
@@ -536,7 +536,7 @@ private:
 		const FString& ProxyBasePackageName,
 		TArray<UObject*>& OutAssetsToSync,
 		FVector& OutProxyLocation
-		) OVERRIDE;
+		) override;
 	
 	bool ConstructRawMesh(
 		UStaticMeshComponent* MeshComponent, 
@@ -3379,7 +3379,7 @@ private:
 			return FindSlotForTexture(InWidth, InWidth);
 		}
 
-		virtual void ConditionalUpdateTexture() OVERRIDE {};
+		virtual void ConditionalUpdateTexture() override {};
 	};
 
 private:
diff --git a/Engine/Source/Developer/MessageLog/Private/Presentation/MessageLogListingViewModel.h b/Engine/Source/Developer/MessageLog/Private/Presentation/MessageLogListingViewModel.h
--- a/Engine/Source/Developer/MessageLog/Private/Presentation/MessageLogListingViewModel.h
+++ b/Engine/Source/Developer/MessageLog/Private/Presentation/MessageLogListingViewModel.h
@@ -32,40 +32,40 @@ public:
 	~FMessageLogListingViewModel();
 
 	/** Begin IMessageLogListing interface */
-	virtual void AddMessage( const TSharedRef< class FTokenizedMessage >& NewMessage ) OVERRIDE;
-	virtual void AddMessages( const TArray< TSharedRef< class FTokenizedMessage > >& NewMessages ) OVERRIDE;
-	virtual void ClearMessages() OVERRIDE;
-	virtual TSharedPtr<FTokenizedMessage> GetMessageFromData(const struct FTokenizedMiscData& MessageData) const OVERRIDE;
-	virtual const TArray< TSharedRef<class FTokenizedMessage> >& GetSelectedMessages() const OVERRIDE;
-	virtual void SelectMessages( const TArray< TSharedRef<class FTokenizedMessage> >& InSelectedMessages ) OVERRIDE;
-	virtual const TArray< TSharedRef<class FTokenizedMessage> >& GetFilteredMessages() const OVERRIDE;
-	virtual void SelectMessage(const TSharedRef<class FTokenizedMessage>& Message, bool bSelected) OVERRIDE;
-	virtual bool IsMessageSelected(const TSharedRef<class FTokenizedMessage>& Message) const OVERRIDE;
-	virtual void ClearSelectedMessages() OVERRIDE;
-	virtual void InvertSelectedMessages() OVERRIDE;
-	virtual FText GetSelectedMessagesAsText() const OVERRIDE;
-	virtual FText GetAllMessagesAsText() const OVERRIDE;
-	virtual const FName& GetName() const OVERRIDE;
-	virtual void SetLabel( const FText& InLogLabel ) OVERRIDE;
-	virtual const FText& GetLabel() const OVERRIDE;
-	virtual const TArray< TSharedRef< class FMessageFilter> >& GetMessageFilters() const OVERRIDE;
-	virtual void ExecuteToken( const TSharedRef<class IMessageToken>& Token ) const OVERRIDE;
-	virtual void NewPage( const FText& Title ) OVERRIDE;
-	virtual void NotifyIfAnyMessages( const FText& Message, EMessageSeverity::Type SeverityFilter = EMessageSeverity::Info, bool bForce = false ) OVERRIDE;
-	virtual void Open() OVERRIDE;
-	virtual int32 NumMessages( EMessageSeverity::Type SeverityFilter ) OVERRIDE;
+	virtual void AddMessage( const TSharedRef< class FTokenizedMessage >& NewMessage ) override;
+	virtual void AddMessages( const TArray< TSharedRef< class FTokenizedMessage > >& NewMessages ) override;
+	virtual void ClearMessages() override;
+	virtual TSharedPtr<FTokenizedMessage> GetMessageFromData(const struct FTokenizedMiscData& MessageData) const override;
+	virtual const TArray< TSharedRef<class FTokenizedMessage> >& GetSelectedMessages() const override;
+	virtual void SelectMessages( const TArray< TSharedRef<class FTokenizedMessage> >& InSelectedMessages ) override;
+	virtual const TArray< TSharedRef<class FTokenizedMessage> >& GetFilteredMessages() const override;
+	virtual void SelectMessage(const TSharedRef<class FTokenizedMessage>& Message, bool bSelected) override;
+	virtual bool IsMessageSelected(const TSharedRef<class FTokenizedMessage>& Message) const override;
+	virtual void ClearSelectedMessages() override;
+	virtual void InvertSelectedMessages() override;
+	virtual FText GetSelectedMessagesAsText() const override;
+	virtual FText GetAllMessagesAsText() const override;
+	virtual const FName& GetName() const override;
+	virtual void SetLabel( const FText& InLogLabel ) override;
+	virtual const FText& GetLabel() const override;
+	virtual const TArray< TSharedRef< class FMessageFilter> >& GetMessageFilters() const override;
+	virtual void ExecuteToken( const TSharedRef<class IMessageToken>& Token ) const override;
+	virtual void NewPage( const FText& Title ) override;
+	virtual void NotifyIfAnyMessages( const FText& Message, EMessageSeverity::Type SeverityFilter = EMessageSeverity::Info, bool bForce = false ) override;
+	virtual void Open() override;
+	virtual int32 NumMessages( EMessageSeverity::Type SeverityFilter ) override;
 
 	DECLARE_DERIVED_EVENT(FMessageLogListingViewModel, IMessageLogListing::IMessageTokenClickedEvent, IMessageTokenClickedEvent)
-	virtual IMessageLogListing::IMessageTokenClickedEvent& OnMessageTokenClicked() OVERRIDE { return TokenClickedEvent; }
+	virtual IMessageLogListing::IMessageTokenClickedEvent& OnMessageTokenClicked() override { return TokenClickedEvent; }
 
 	DECLARE_DERIVED_EVENT(FMessageLogListingViewModel, IMessageLogListing::FChangedEvent, FChangedEvent)
-	virtual IMessageLogListing::FChangedEvent& OnDataChanged() OVERRIDE { return ChangedEvent; }
+	virtual IMessageLogListing::FChangedEvent& OnDataChanged() override { return ChangedEvent; }
 
 	DECLARE_DERIVED_EVENT(FMessageLogListingViewModel, IMessageLogListing::FOnSelectionChangedEvent, FOnSelectionChangedEvent)
-	virtual IMessageLogListing::FOnSelectionChangedEvent& OnSelectionChanged() OVERRIDE { return SelectionChangedEvent; }
+	virtual IMessageLogListing::FOnSelectionChangedEvent& OnSelectionChanged() override { return SelectionChangedEvent; }
 
 	DECLARE_DERIVED_EVENT(FMessageLogListingViewModel, IMessageLogListing::FOnPageSelectionChangedEvent, FOnPageSelectionChangedEvent )
-	virtual IMessageLogListing::FOnPageSelectionChangedEvent& OnPageSelectionChanged() OVERRIDE { return PageSelectionChangedEvent; }
+	virtual IMessageLogListing::FOnPageSelectionChangedEvent& OnPageSelectionChanged() override { return PageSelectionChangedEvent; }
 	/** End IMessageLogListing interface */
 
 	/** Initializes the FMessageLogListingViewModel for use */
diff --git a/Engine/Source/Developer/MessageLog/Private/UserInterface/SMessageLog.h b/Engine/Source/Developer/MessageLog/Private/UserInterface/SMessageLog.h
--- a/Engine/Source/Developer/MessageLog/Private/UserInterface/SMessageLog.h
+++ b/Engine/Source/Developer/MessageLog/Private/UserInterface/SMessageLog.h
@@ -29,7 +29,7 @@ public:
 	 *
 	 * @return  Returns whether the event was handled, along with other possible actions
 	 */
-	FReply OnKeyDown( const FGeometry& MyGeometry, const FKeyboardEvent& InKeyboardEvent ) OVERRIDE;
+	FReply OnKeyDown( const FGeometry& MyGeometry, const FKeyboardEvent& InKeyboardEvent ) override;
 
 private:
 	/** Gets the currently visible log listing's label */
diff --git a/Engine/Source/Developer/MessageLog/Private/UserInterface/SMessageLogListing.h b/Engine/Source/Developer/MessageLog/Private/UserInterface/SMessageLogListing.h
--- a/Engine/Source/Developer/MessageLog/Private/UserInterface/SMessageLogListing.h
+++ b/Engine/Source/Developer/MessageLog/Private/UserInterface/SMessageLogListing.h
@@ -87,7 +87,7 @@ public:
 	 *
 	 * @return  Returns whether the event was handled, along with other possible actions
 	 */
-	FReply OnKeyDown( const FGeometry& MyGeometry, const FKeyboardEvent& InKeyboardEvent ) OVERRIDE;
+	FReply OnKeyDown( const FGeometry& MyGeometry, const FKeyboardEvent& InKeyboardEvent ) override;
 
 	/** Delegate supplying a label for the page-flipper widget */
 	FString GetPageText() const;
diff --git a/Engine/Source/Developer/MovieSceneTools/Classes/K2Node_PlayMovieScene.h b/Engine/Source/Developer/MovieSceneTools/Classes/K2Node_PlayMovieScene.h
--- a/Engine/Source/Developer/MovieSceneTools/Classes/K2Node_PlayMovieScene.h
+++ b/Engine/Source/Developer/MovieSceneTools/Classes/K2Node_PlayMovieScene.h
@@ -26,7 +26,7 @@ struct FEdGraphSchemaAction_K2AddPlayMovieScene : public FEdGraphSchemaAction_K2
 	}
 	
 	// FEdGraphSchemaAction interface
-	virtual UEdGraphNode* PerformAction(class UEdGraph* ParentGraph, UEdGraphPin* FromPin, const FVector2D Location, bool bSelectNewNode = true) OVERRIDE;
+	virtual UEdGraphNode* PerformAction(class UEdGraph* ParentGraph, UEdGraphPin* FromPin, const FVector2D Location, bool bSelectNewNode = true) override;
 	// End of FEdGraphSchemaAction interface
 
 };
@@ -39,18 +39,18 @@ class UK2Node_PlayMovieScene : public UK2Node
 
 #if WITH_EDITOR
 	// UEdGraphNode interface.
-	virtual void AllocateDefaultPins() OVERRIDE;
-	virtual FLinearColor GetNodeTitleColor() const OVERRIDE;
-	virtual FText GetNodeTitle(ENodeTitleType::Type TitleType) const OVERRIDE;
-	virtual FString GetNodeNativeTitle(ENodeTitleType::Type TitleType) const OVERRIDE;
+	virtual void AllocateDefaultPins() override;
+	virtual FLinearColor GetNodeTitleColor() const override;
+	virtual FText GetNodeTitle(ENodeTitleType::Type TitleType) const override;
+	virtual FString GetNodeNativeTitle(ENodeTitleType::Type TitleType) const override;
 
 	// UK2Node interface
-	virtual bool NodeCausesStructuralBlueprintChange() const OVERRIDE
+	virtual bool NodeCausesStructuralBlueprintChange() const override
 	{ 
 		// @todo sequencer: Not sure if we really need to return true for this
 		return true; 
 	}
-	virtual void PinConnectionListChanged( UEdGraphPin* Pin ) OVERRIDE;
+	virtual void PinConnectionListChanged( UEdGraphPin* Pin ) override;
 
 
 	/**
diff --git a/Engine/Source/Developer/MovieSceneTools/Classes/MovieSceneFactory.h b/Engine/Source/Developer/MovieSceneTools/Classes/MovieSceneFactory.h
--- a/Engine/Source/Developer/MovieSceneTools/Classes/MovieSceneFactory.h
+++ b/Engine/Source/Developer/MovieSceneTools/Classes/MovieSceneFactory.h
@@ -10,7 +10,7 @@ class UMovieSceneFactory : public UFactory
 	GENERATED_UCLASS_BODY()
 
 	// Begin UFactory Interface
-	virtual UObject* FactoryCreateNew(UClass* Class,UObject* InParent,FName Name,EObjectFlags Flags,UObject* Context,FFeedbackContext* Warn) OVERRIDE;
+	virtual UObject* FactoryCreateNew(UClass* Class,UObject* InParent,FName Name,EObjectFlags Flags,UObject* Context,FFeedbackContext* Warn) override;
 	// Begin UFactory Interface
 };
 
diff --git a/Engine/Source/Developer/MovieSceneTools/Private/AnimationTrackEditor.h b/Engine/Source/Developer/MovieSceneTools/Private/AnimationTrackEditor.h
--- a/Engine/Source/Developer/MovieSceneTools/Private/AnimationTrackEditor.h
+++ b/Engine/Source/Developer/MovieSceneTools/Private/AnimationTrackEditor.h
@@ -26,11 +26,11 @@ public:
 	static TSharedRef<FMovieSceneTrackEditor> CreateTrackEditor( TSharedRef<ISequencer> OwningSequencer );
 
 	/** FMovieSceneTrackEditor Interface */
-	virtual bool SupportsType( TSubclassOf<UMovieSceneTrack> Type ) const OVERRIDE;
-	virtual TSharedRef<ISequencerSection> MakeSectionInterface( UMovieSceneSection& SectionObject, UMovieSceneTrack* Track ) OVERRIDE;
-	virtual void AddKey(const FGuid& ObjectGuid, UObject* AdditionalAsset) OVERRIDE;
-	virtual bool HandleAssetAdded(UObject* Asset, const FGuid& TargetObjectGuid) OVERRIDE;
-	virtual void BuildObjectBindingContextMenu(FMenuBuilder& MenuBuilder, const FGuid& ObjectBinding, const UClass* ObjectClass) OVERRIDE;
+	virtual bool SupportsType( TSubclassOf<UMovieSceneTrack> Type ) const override;
+	virtual TSharedRef<ISequencerSection> MakeSectionInterface( UMovieSceneSection& SectionObject, UMovieSceneTrack* Track ) override;
+	virtual void AddKey(const FGuid& ObjectGuid, UObject* AdditionalAsset) override;
+	virtual bool HandleAssetAdded(UObject* Asset, const FGuid& TargetObjectGuid) override;
+	virtual void BuildObjectBindingContextMenu(FMenuBuilder& MenuBuilder, const FGuid& ObjectBinding, const UClass* ObjectClass) override;
 
 private:
 	/** Delegate for AnimatablePropertyChanged in AddKey */
@@ -50,13 +50,13 @@ public:
 	~FAnimationSection();
 
 	/** ISequencerSection interface */
-	virtual UMovieSceneSection* GetSectionObject() OVERRIDE;
-	virtual FString GetDisplayName() const OVERRIDE;
-	virtual FString GetSectionTitle() const OVERRIDE;
-	virtual float GetSectionHeight() const OVERRIDE;
-	virtual void GenerateSectionLayout( class ISectionLayoutBuilder& LayoutBuilder ) const OVERRIDE {}
-	virtual int32 OnPaintSection( const FGeometry& AllottedGeometry, const FSlateRect& SectionClippingRect, FSlateWindowElementList& OutDrawElements, int32 LayerId, bool bParentEnabled ) const OVERRIDE;
-	virtual void Tick( const FGeometry& AllottedGeometry, const FGeometry& ParentGeometry, const double InCurrentTime, const float InDeltaTime ) OVERRIDE;
+	virtual UMovieSceneSection* GetSectionObject() override;
+	virtual FString GetDisplayName() const override;
+	virtual FString GetSectionTitle() const override;
+	virtual float GetSectionHeight() const override;
+	virtual void GenerateSectionLayout( class ISectionLayoutBuilder& LayoutBuilder ) const override {}
+	virtual int32 OnPaintSection( const FGeometry& AllottedGeometry, const FSlateRect& SectionClippingRect, FSlateWindowElementList& OutDrawElements, int32 LayerId, bool bParentEnabled ) const override;
+	virtual void Tick( const FGeometry& AllottedGeometry, const FGeometry& ParentGeometry, const double InCurrentTime, const float InDeltaTime ) override;
 
 private:
 	/** The section we are visualizing */
diff --git a/Engine/Source/Developer/MovieSceneTools/Private/AssetTypeActions_MovieScene.h b/Engine/Source/Developer/MovieSceneTools/Private/AssetTypeActions_MovieScene.h
--- a/Engine/Source/Developer/MovieSceneTools/Private/AssetTypeActions_MovieScene.h
+++ b/Engine/Source/Developer/MovieSceneTools/Private/AssetTypeActions_MovieScene.h
@@ -8,15 +8,15 @@ class FAssetTypeActions_MovieScene : public FAssetTypeActions_Base
 {
 public:
 	// IAssetTypeActions Implementation
-	virtual FText GetName() const OVERRIDE { return NSLOCTEXT("AssetTypeActions", "AssetTypeActions_MovieScene", "Movie Scene"); }
-	virtual FColor GetTypeColor() const OVERRIDE { return FColor(200, 80, 80); }
-	virtual UClass* GetSupportedClass() const OVERRIDE;
-	virtual bool HasActions ( const TArray<UObject*>& InObjects ) const OVERRIDE { return true; }
-	virtual void GetActions( const TArray<UObject*>& InObjects, FMenuBuilder& MenuBuilder ) OVERRIDE;
-	virtual void OpenAssetEditor( const TArray<UObject*>& InObjects, TSharedPtr<class IToolkitHost> EditWithinLevelEditor = TSharedPtr<IToolkitHost>() ) OVERRIDE;
-	virtual uint32 GetCategories() OVERRIDE { return EAssetTypeCategories::Animation; }
+	virtual FText GetName() const override { return NSLOCTEXT("AssetTypeActions", "AssetTypeActions_MovieScene", "Movie Scene"); }
+	virtual FColor GetTypeColor() const override { return FColor(200, 80, 80); }
+	virtual UClass* GetSupportedClass() const override;
+	virtual bool HasActions ( const TArray<UObject*>& InObjects ) const override { return true; }
+	virtual void GetActions( const TArray<UObject*>& InObjects, FMenuBuilder& MenuBuilder ) override;
+	virtual void OpenAssetEditor( const TArray<UObject*>& InObjects, TSharedPtr<class IToolkitHost> EditWithinLevelEditor = TSharedPtr<IToolkitHost>() ) override;
+	virtual uint32 GetCategories() override { return EAssetTypeCategories::Animation; }
 
-	virtual bool ShouldForceWorldCentric() OVERRIDE
+	virtual bool ShouldForceWorldCentric() override
 	{
 		// @todo sequencer: Hack to force world-centric mode for Sequencer
 		return true;
diff --git a/Engine/Source/Developer/MovieSceneTools/Private/AudioTrackEditor.cpp b/Engine/Source/Developer/MovieSceneTools/Private/AudioTrackEditor.cpp
--- a/Engine/Source/Developer/MovieSceneTools/Private/AudioTrackEditor.cpp
+++ b/Engine/Source/Developer/MovieSceneTools/Private/AudioTrackEditor.cpp
@@ -88,9 +88,9 @@ public:
 	~FAudioThumbnail();
 
 	/* ISlateViewport interface */
-	virtual FIntPoint GetSize() const OVERRIDE;
-	virtual class FSlateShaderResource* GetViewportRenderTargetTexture() const OVERRIDE;
-	virtual bool RequiresVsync() const OVERRIDE;
+	virtual FIntPoint GetSize() const override;
+	virtual class FSlateShaderResource* GetViewportRenderTargetTexture() const override;
+	virtual bool RequiresVsync() const override;
 
 	/** Returns whether this thumbnail has a texture to render */
 	virtual bool ShouldRender() const {return TextureSize > 0;}
diff --git a/Engine/Source/Developer/MovieSceneTools/Private/AudioTrackEditor.h b/Engine/Source/Developer/MovieSceneTools/Private/AudioTrackEditor.h
--- a/Engine/Source/Developer/MovieSceneTools/Private/AudioTrackEditor.h
+++ b/Engine/Source/Developer/MovieSceneTools/Private/AudioTrackEditor.h
@@ -26,9 +26,9 @@ public:
 	static TSharedRef<FMovieSceneTrackEditor> CreateTrackEditor( TSharedRef<ISequencer> OwningSequencer );
 
 	/** FMovieSceneTrackEditor Interface */
-	virtual bool SupportsType( TSubclassOf<UMovieSceneTrack> Type ) const OVERRIDE;
-	virtual TSharedRef<ISequencerSection> MakeSectionInterface( UMovieSceneSection& SectionObject, UMovieSceneTrack* Track ) OVERRIDE;
-	virtual bool HandleAssetAdded(UObject* Asset, const FGuid& TargetObjectGuid) OVERRIDE;
+	virtual bool SupportsType( TSubclassOf<UMovieSceneTrack> Type ) const override;
+	virtual TSharedRef<ISequencerSection> MakeSectionInterface( UMovieSceneSection& SectionObject, UMovieSceneTrack* Track ) override;
+	virtual bool HandleAssetAdded(UObject* Asset, const FGuid& TargetObjectGuid) override;
 	
 private:
 	/** Delegate for AnimatablePropertyChanged in HandleAssetAdded for master sounds */
@@ -47,13 +47,13 @@ public:
 	~FAudioSection();
 
 	/** ISequencerSection interface */
-	virtual UMovieSceneSection* GetSectionObject() OVERRIDE;
-	virtual FString GetDisplayName() const OVERRIDE;
-	virtual FString GetSectionTitle() const OVERRIDE;
-	virtual float GetSectionHeight() const OVERRIDE;
-	virtual void GenerateSectionLayout( class ISectionLayoutBuilder& LayoutBuilder ) const OVERRIDE {}
-	virtual int32 OnPaintSection( const FGeometry& AllottedGeometry, const FSlateRect& SectionClippingRect, FSlateWindowElementList& OutDrawElements, int32 LayerId, bool bParentEnabled ) const OVERRIDE;
-	virtual void Tick( const FGeometry& AllottedGeometry, const FGeometry& ParentGeometry, const double InCurrentTime, const float InDeltaTime ) OVERRIDE;
+	virtual UMovieSceneSection* GetSectionObject() override;
+	virtual FString GetDisplayName() const override;
+	virtual FString GetSectionTitle() const override;
+	virtual float GetSectionHeight() const override;
+	virtual void GenerateSectionLayout( class ISectionLayoutBuilder& LayoutBuilder ) const override {}
+	virtual int32 OnPaintSection( const FGeometry& AllottedGeometry, const FSlateRect& SectionClippingRect, FSlateWindowElementList& OutDrawElements, int32 LayerId, bool bParentEnabled ) const override;
+	virtual void Tick( const FGeometry& AllottedGeometry, const FGeometry& ParentGeometry, const double InCurrentTime, const float InDeltaTime ) override;
 	
 private:
 	/* Re-creates the texture used to preview the waveform */
diff --git a/Engine/Source/Developer/MovieSceneTools/Private/DirectorTrackEditor.h b/Engine/Source/Developer/MovieSceneTools/Private/DirectorTrackEditor.h
--- a/Engine/Source/Developer/MovieSceneTools/Private/DirectorTrackEditor.h
+++ b/Engine/Source/Developer/MovieSceneTools/Private/DirectorTrackEditor.h
@@ -26,11 +26,11 @@ public:
 	static TSharedRef<FMovieSceneTrackEditor> CreateTrackEditor( TSharedRef<ISequencer> OwningSequencer );
 
 	/** FMovieSceneTrackEditor Interface */
-	virtual bool SupportsType( TSubclassOf<UMovieSceneTrack> Type ) const OVERRIDE;
-	virtual TSharedRef<ISequencerSection> MakeSectionInterface( UMovieSceneSection& SectionObject, UMovieSceneTrack* Track ) OVERRIDE;
-	virtual void AddKey(const FGuid& ObjectGuid, UObject* AdditionalAsset = NULL) OVERRIDE;
-	virtual void Tick(float DeltaTime) OVERRIDE;
-	virtual void BuildObjectBindingContextMenu(FMenuBuilder& MenuBuilder, const FGuid& ObjectBinding, const UClass* ObjectClass) OVERRIDE;
+	virtual bool SupportsType( TSubclassOf<UMovieSceneTrack> Type ) const override;
+	virtual TSharedRef<ISequencerSection> MakeSectionInterface( UMovieSceneSection& SectionObject, UMovieSceneTrack* Track ) override;
+	virtual void AddKey(const FGuid& ObjectGuid, UObject* AdditionalAsset = NULL) override;
+	virtual void Tick(float DeltaTime) override;
+	virtual void BuildObjectBindingContextMenu(FMenuBuilder& MenuBuilder, const FGuid& ObjectBinding, const UClass* ObjectClass) override;
 
 private:
 	/** Delegate for AnimatablePropertyChanged in AddKey */
@@ -82,9 +82,9 @@ public:
 	~FShotThumbnail();
 
 	/* ISlateViewport interface */
-	virtual FIntPoint GetSize() const OVERRIDE;
-	virtual class FSlateShaderResource* GetViewportRenderTargetTexture() const OVERRIDE;
-	virtual bool RequiresVsync() const OVERRIDE;
+	virtual FIntPoint GetSize() const override;
+	virtual class FSlateShaderResource* GetViewportRenderTargetTexture() const override;
+	virtual bool RequiresVsync() const override;
 
 	/** Gets the time that this thumbnail is a rendering of */
 	float GetTime() const;
@@ -129,14 +129,14 @@ public:
 	~FShotSection();
 
 	/** ISequencerSection interface */
-	virtual UMovieSceneSection* GetSectionObject() OVERRIDE;
-	virtual int32 OnPaintSection( const FGeometry& AllottedGeometry, const FSlateRect& SectionClippingRect, FSlateWindowElementList& OutDrawElements, int32 LayerId, bool bParentEnabled ) const OVERRIDE;
-	virtual void Tick( const FGeometry& AllottedGeometry, const FGeometry& ParentGeometry, const double InCurrentTime, const float InDeltaTime ) OVERRIDE;
-	virtual FString GetDisplayName() const OVERRIDE { return TEXT("Director Track"); }
-	virtual FString GetSectionTitle() const OVERRIDE;
-	virtual float GetSectionHeight() const OVERRIDE;
-	virtual void GenerateSectionLayout( class ISectionLayoutBuilder& LayoutBuilder ) const OVERRIDE {}
-	virtual FReply OnSectionDoubleClicked( const FGeometry& SectionGeometry, const FPointerEvent& MouseEvent ) OVERRIDE;
+	virtual UMovieSceneSection* GetSectionObject() override;
+	virtual int32 OnPaintSection( const FGeometry& AllottedGeometry, const FSlateRect& SectionClippingRect, FSlateWindowElementList& OutDrawElements, int32 LayerId, bool bParentEnabled ) const override;
+	virtual void Tick( const FGeometry& AllottedGeometry, const FGeometry& ParentGeometry, const double InCurrentTime, const float InDeltaTime ) override;
+	virtual FString GetDisplayName() const override { return TEXT("Director Track"); }
+	virtual FString GetSectionTitle() const override;
+	virtual float GetSectionHeight() const override;
+	virtual void GenerateSectionLayout( class ISectionLayoutBuilder& LayoutBuilder ) const override {}
+	virtual FReply OnSectionDoubleClicked( const FGeometry& SectionGeometry, const FPointerEvent& MouseEvent ) override;
 
 	/** Gets the thumbnail width */
 	uint32 GetThumbnailWidth() const;
diff --git a/Engine/Source/Developer/MovieSceneTools/Private/MovieSceneToolHelpers.h b/Engine/Source/Developer/MovieSceneTools/Private/MovieSceneToolHelpers.h
--- a/Engine/Source/Developer/MovieSceneTools/Private/MovieSceneToolHelpers.h
+++ b/Engine/Source/Developer/MovieSceneTools/Private/MovieSceneToolHelpers.h
@@ -15,7 +15,7 @@ public:
 	{}
 
 	/** IKeyArea interface */
-	virtual TArray<FKeyHandle> GetUnsortedKeyHandles() const OVERRIDE
+	virtual TArray<FKeyHandle> GetUnsortedKeyHandles() const override
 	{
 		TArray<FKeyHandle> OutKeyHandles;
 		for (auto It(Curve.GetKeyHandleIterator()); It; ++It)
@@ -25,17 +25,17 @@ public:
 		return OutKeyHandles;
 	}
 
-	virtual float GetKeyTime( FKeyHandle KeyHandle ) const OVERRIDE
+	virtual float GetKeyTime( FKeyHandle KeyHandle ) const override
 	{
 		return Curve.GetKeyTime( KeyHandle );
 	}
 
-	virtual FKeyHandle MoveKey( FKeyHandle KeyHandle, float DeltaPosition ) OVERRIDE
+	virtual FKeyHandle MoveKey( FKeyHandle KeyHandle, float DeltaPosition ) override
 	{
 		return Curve.SetKeyTime( KeyHandle, Curve.GetKeyTime( KeyHandle ) + DeltaPosition );
 	}
 
-	virtual void DeleteKey(FKeyHandle KeyHandle) OVERRIDE
+	virtual void DeleteKey(FKeyHandle KeyHandle) override
 	{
 		Curve.DeleteKey(KeyHandle);
 	}
@@ -55,7 +55,7 @@ public:
 	{}
 
 	/** IKeyArea interface */
-	virtual TArray<FKeyHandle> GetUnsortedKeyHandles() const OVERRIDE
+	virtual TArray<FKeyHandle> GetUnsortedKeyHandles() const override
 	{
 		TArray<FKeyHandle> OutKeyHandles;
 		for (auto It(Curve.GetKeyHandleIterator()); It; ++It)
@@ -65,17 +65,17 @@ public:
 		return OutKeyHandles;
 	}
 
-	virtual float GetKeyTime( FKeyHandle KeyHandle ) const OVERRIDE
+	virtual float GetKeyTime( FKeyHandle KeyHandle ) const override
 	{
 		return Curve.GetKeyTime( KeyHandle );
 	}
 
-	virtual FKeyHandle MoveKey( FKeyHandle KeyHandle, float DeltaPosition ) OVERRIDE
+	virtual FKeyHandle MoveKey( FKeyHandle KeyHandle, float DeltaPosition ) override
 	{
 		return Curve.SetKeyTime( KeyHandle, Curve.GetKeyTime( KeyHandle ) + DeltaPosition );
 	}
 
-	virtual void DeleteKey(FKeyHandle KeyHandle) OVERRIDE
+	virtual void DeleteKey(FKeyHandle KeyHandle) override
 	{
 		Curve.DeleteKey(KeyHandle);
 	}
diff --git a/Engine/Source/Developer/MovieSceneTools/Private/ParticleTrackEditor.h b/Engine/Source/Developer/MovieSceneTools/Private/ParticleTrackEditor.h
--- a/Engine/Source/Developer/MovieSceneTools/Private/ParticleTrackEditor.h
+++ b/Engine/Source/Developer/MovieSceneTools/Private/ParticleTrackEditor.h
@@ -26,9 +26,9 @@ public:
 	static TSharedRef<FMovieSceneTrackEditor> CreateTrackEditor( TSharedRef<ISequencer> OwningSequencer );
 
 	/** FMovieSceneTrackEditor Interface */
-	virtual bool SupportsType( TSubclassOf<UMovieSceneTrack> Type ) const OVERRIDE;
-	virtual TSharedRef<ISequencerSection> MakeSectionInterface( UMovieSceneSection& SectionObject, UMovieSceneTrack* Track ) OVERRIDE;
-	virtual void BuildObjectBindingContextMenu(FMenuBuilder& MenuBuilder, const FGuid& ObjectBinding, const UClass* ObjectClass) OVERRIDE;
+	virtual bool SupportsType( TSubclassOf<UMovieSceneTrack> Type ) const override;
+	virtual TSharedRef<ISequencerSection> MakeSectionInterface( UMovieSceneSection& SectionObject, UMovieSceneTrack* Track ) override;
+	virtual void BuildObjectBindingContextMenu(FMenuBuilder& MenuBuilder, const FGuid& ObjectBinding, const UClass* ObjectClass) override;
 	
 	void AddParticleKey(const FGuid ObjectGuid, bool bTrigger);
 
@@ -47,13 +47,13 @@ public:
 	~FParticleSection();
 
 	/** ISequencerSection interface */
-	virtual UMovieSceneSection* GetSectionObject() OVERRIDE;
-	virtual FString GetDisplayName() const OVERRIDE;
-	virtual FString GetSectionTitle() const OVERRIDE;
-	virtual float GetSectionHeight() const OVERRIDE;
-	virtual void GenerateSectionLayout( class ISectionLayoutBuilder& LayoutBuilder ) const OVERRIDE {}
-	virtual int32 OnPaintSection( const FGeometry& AllottedGeometry, const FSlateRect& SectionClippingRect, FSlateWindowElementList& OutDrawElements, int32 LayerId, bool bParentEnabled ) const OVERRIDE;
-	virtual bool SectionIsResizable() const OVERRIDE;
+	virtual UMovieSceneSection* GetSectionObject() override;
+	virtual FString GetDisplayName() const override;
+	virtual FString GetSectionTitle() const override;
+	virtual float GetSectionHeight() const override;
+	virtual void GenerateSectionLayout( class ISectionLayoutBuilder& LayoutBuilder ) const override {}
+	virtual int32 OnPaintSection( const FGeometry& AllottedGeometry, const FSlateRect& SectionClippingRect, FSlateWindowElementList& OutDrawElements, int32 LayerId, bool bParentEnabled ) const override;
+	virtual bool SectionIsResizable() const override;
 
 private:
 	/** The section we are visualizing */
diff --git a/Engine/Source/Developer/MovieSceneTools/Private/PropertyTrackEditor.cpp b/Engine/Source/Developer/MovieSceneTools/Private/PropertyTrackEditor.cpp
--- a/Engine/Source/Developer/MovieSceneTools/Private/PropertyTrackEditor.cpp
+++ b/Engine/Source/Developer/MovieSceneTools/Private/PropertyTrackEditor.cpp
@@ -35,18 +35,18 @@ public:
 	}
 
 	/** ISequencerSection interface */
-	virtual UMovieSceneSection* GetSectionObject() OVERRIDE { return &SectionObject; }
+	virtual UMovieSceneSection* GetSectionObject() override { return &SectionObject; }
 
-	virtual FString GetDisplayName() const OVERRIDE
+	virtual FString GetDisplayName() const override
 	{
 		return DisplayName;
 	}
 	
-	virtual FString GetSectionTitle() const OVERRIDE { return FString(); }
+	virtual FString GetSectionTitle() const override { return FString(); }
 
-	virtual void GenerateSectionLayout( class ISectionLayoutBuilder& LayoutBuilder ) const OVERRIDE {}
+	virtual void GenerateSectionLayout( class ISectionLayoutBuilder& LayoutBuilder ) const override {}
 
-	virtual int32 OnPaintSection( const FGeometry& AllottedGeometry, const FSlateRect& SectionClippingRect, FSlateWindowElementList& OutDrawElements, int32 LayerId, bool bParentEnabled ) const OVERRIDE 
+	virtual int32 OnPaintSection( const FGeometry& AllottedGeometry, const FSlateRect& SectionClippingRect, FSlateWindowElementList& OutDrawElements, int32 LayerId, bool bParentEnabled ) const override 
 	{
 		// Add a box for the section
 		FSlateDrawElement::MakeBox( 
@@ -76,7 +76,7 @@ public:
 	FColorPropertySection( UMovieSceneSection& InSectionObject, FName SectionName )
 		: FPropertySection(InSectionObject, SectionName) {}
 
-	virtual void GenerateSectionLayout( class ISectionLayoutBuilder& LayoutBuilder ) const OVERRIDE
+	virtual void GenerateSectionLayout( class ISectionLayoutBuilder& LayoutBuilder ) const override
 	{
 		UMovieSceneColorSection* ColorSection = Cast<UMovieSceneColorSection>(&SectionObject);
 
@@ -86,7 +86,7 @@ public:
 		LayoutBuilder.AddKeyArea("A", TEXT("Alpha"), MakeShareable(new FFloatCurveKeyArea(ColorSection->GetAlphaCurve())));
 	}
 
-	virtual int32 OnPaintSection( const FGeometry& AllottedGeometry, const FSlateRect& SectionClippingRect, FSlateWindowElementList& OutDrawElements, int32 LayerId, bool bParentEnabled ) const OVERRIDE 
+	virtual int32 OnPaintSection( const FGeometry& AllottedGeometry, const FSlateRect& SectionClippingRect, FSlateWindowElementList& OutDrawElements, int32 LayerId, bool bParentEnabled ) const override 
 	{
 		const UMovieSceneColorSection* ColorSection = Cast<const UMovieSceneColorSection>(&SectionObject);
 
@@ -198,7 +198,7 @@ public:
 	FFloatPropertySection( UMovieSceneSection& InSectionObject, FName SectionName )
 		: FPropertySection(InSectionObject, SectionName) {}
 
-	virtual void GenerateSectionLayout( class ISectionLayoutBuilder& LayoutBuilder ) const OVERRIDE
+	virtual void GenerateSectionLayout( class ISectionLayoutBuilder& LayoutBuilder ) const override
 	{
 		UMovieSceneFloatSection* FloatSection = Cast<UMovieSceneFloatSection>( &SectionObject );
 		LayoutBuilder.SetSectionAsKeyArea( MakeShareable( new FFloatCurveKeyArea( FloatSection->GetFloatCurve() ) ) );
@@ -214,7 +214,7 @@ public:
 	FVectorPropertySection( UMovieSceneSection& InSectionObject, FName SectionName )
 		: FPropertySection(InSectionObject, SectionName) {}
 
-	virtual void GenerateSectionLayout( class ISectionLayoutBuilder& LayoutBuilder ) const OVERRIDE
+	virtual void GenerateSectionLayout( class ISectionLayoutBuilder& LayoutBuilder ) const override
 	{
 		UMovieSceneVectorSection* VectorSection = Cast<UMovieSceneVectorSection>( &SectionObject );
 		int32 ChannelsUsed = VectorSection->GetChannelsUsed();
@@ -238,13 +238,13 @@ public:
 	FBoolPropertySection( UMovieSceneSection& InSectionObject, FName SectionName )
 		: FPropertySection(InSectionObject, SectionName) {}
 
-	virtual void GenerateSectionLayout( class ISectionLayoutBuilder& LayoutBuilder ) const OVERRIDE
+	virtual void GenerateSectionLayout( class ISectionLayoutBuilder& LayoutBuilder ) const override
 	{
 		UMovieSceneBoolSection* BoolSection = Cast<UMovieSceneBoolSection>( &SectionObject );
 		LayoutBuilder.SetSectionAsKeyArea( MakeShareable( new FIntegralKeyArea( BoolSection->GetCurve() ) ) );
 	}
 
-	virtual int32 OnPaintSection( const FGeometry& AllottedGeometry, const FSlateRect& SectionClippingRect, FSlateWindowElementList& OutDrawElements, int32 LayerId, bool bParentEnabled ) const OVERRIDE 
+	virtual int32 OnPaintSection( const FGeometry& AllottedGeometry, const FSlateRect& SectionClippingRect, FSlateWindowElementList& OutDrawElements, int32 LayerId, bool bParentEnabled ) const override 
 	{
 		FTimeToPixel TimeToPixelConverter( AllottedGeometry, TRange<float>( SectionObject.GetStartTime(), SectionObject.GetEndTime() ) );
 
diff --git a/Engine/Source/Developer/MovieSceneTools/Private/PropertyTrackEditor.h b/Engine/Source/Developer/MovieSceneTools/Private/PropertyTrackEditor.h
--- a/Engine/Source/Developer/MovieSceneTools/Private/PropertyTrackEditor.h
+++ b/Engine/Source/Developer/MovieSceneTools/Private/PropertyTrackEditor.h
@@ -28,9 +28,9 @@ public:
 	static TSharedRef<FMovieSceneTrackEditor> CreateTrackEditor( TSharedRef<ISequencer> OwningSequencer );
 
 	/** FMovieSceneTrackEditor Interface */
-	virtual bool SupportsType( TSubclassOf<UMovieSceneTrack> Type ) const OVERRIDE;
-	virtual TSharedRef<ISequencerSection> MakeSectionInterface( UMovieSceneSection& SectionObject, UMovieSceneTrack* Track ) OVERRIDE;
-	virtual void AddKey(const FGuid& ObjectGuid, UObject* AdditionalAsset = NULL) OVERRIDE;
+	virtual bool SupportsType( TSubclassOf<UMovieSceneTrack> Type ) const override;
+	virtual TSharedRef<ISequencerSection> MakeSectionInterface( UMovieSceneSection& SectionObject, UMovieSceneTrack* Track ) override;
+	virtual void AddKey(const FGuid& ObjectGuid, UObject* AdditionalAsset = NULL) override;
 
 private:
 	/**
diff --git a/Engine/Source/Developer/MovieSceneTools/Private/SubMovieSceneTrackEditor.cpp b/Engine/Source/Developer/MovieSceneTools/Private/SubMovieSceneTrackEditor.cpp
--- a/Engine/Source/Developer/MovieSceneTools/Private/SubMovieSceneTrackEditor.cpp
+++ b/Engine/Source/Developer/MovieSceneTools/Private/SubMovieSceneTrackEditor.cpp
@@ -20,23 +20,23 @@ public:
 	}
 
 	/** ISequencerSection interface */
-	virtual UMovieSceneSection* GetSectionObject() OVERRIDE { return &SectionObject; }
+	virtual UMovieSceneSection* GetSectionObject() override { return &SectionObject; }
 
-	virtual FString GetDisplayName() const OVERRIDE
+	virtual FString GetDisplayName() const override
 	{
 		return DisplayName;
 	}
 
-	virtual FString GetSectionTitle() const OVERRIDE
+	virtual FString GetSectionTitle() const override
 	{
 		return SectionObject.GetMovieScene()->GetName();
 	}
 
-	virtual void GenerateSectionLayout( class ISectionLayoutBuilder& LayoutBuilder ) const OVERRIDE
+	virtual void GenerateSectionLayout( class ISectionLayoutBuilder& LayoutBuilder ) const override
 	{
 	}
 
-	virtual FReply OnSectionDoubleClicked( const FGeometry& SectionGeometry, const FPointerEvent& MouseEvent ) OVERRIDE
+	virtual FReply OnSectionDoubleClicked( const FGeometry& SectionGeometry, const FPointerEvent& MouseEvent ) override
 	{
 		if( MouseEvent.GetEffectingButton() == EKeys::LeftMouseButton )
 		{
@@ -45,9 +45,9 @@ public:
 		return FReply::Handled();
 	}
 
-	virtual float GetSectionHeight() const OVERRIDE { return 30.0f; }
+	virtual float GetSectionHeight() const override { return 30.0f; }
 
-	virtual int32 OnPaintSection( const FGeometry& AllottedGeometry, const FSlateRect& SectionClippingRect, FSlateWindowElementList& OutDrawElements, int32 LayerId, bool bParentEnabled ) const OVERRIDE 
+	virtual int32 OnPaintSection( const FGeometry& AllottedGeometry, const FSlateRect& SectionClippingRect, FSlateWindowElementList& OutDrawElements, int32 LayerId, bool bParentEnabled ) const override 
 	{
 		// Add a box for the section
 		FSlateDrawElement::MakeBox(
diff --git a/Engine/Source/Developer/MovieSceneTools/Private/SubMovieSceneTrackEditor.h b/Engine/Source/Developer/MovieSceneTools/Private/SubMovieSceneTrackEditor.h
--- a/Engine/Source/Developer/MovieSceneTools/Private/SubMovieSceneTrackEditor.h
+++ b/Engine/Source/Developer/MovieSceneTools/Private/SubMovieSceneTrackEditor.h
@@ -25,9 +25,9 @@ public:
 	static TSharedRef<FMovieSceneTrackEditor> CreateTrackEditor( TSharedRef<ISequencer> OwningSequencer );
 
 	/** FMovieSceneTrackEditor Interface */
-	virtual bool SupportsType( TSubclassOf<UMovieSceneTrack> Type ) const OVERRIDE;
-	virtual TSharedRef<ISequencerSection> MakeSectionInterface( UMovieSceneSection& SectionObject, UMovieSceneTrack* Track ) OVERRIDE;
-	virtual bool HandleAssetAdded(UObject* Asset, const FGuid& TargetObjectGuid) OVERRIDE;
+	virtual bool SupportsType( TSubclassOf<UMovieSceneTrack> Type ) const override;
+	virtual TSharedRef<ISequencerSection> MakeSectionInterface( UMovieSceneSection& SectionObject, UMovieSceneTrack* Track ) override;
+	virtual bool HandleAssetAdded(UObject* Asset, const FGuid& TargetObjectGuid) override;
 };
 
 
