implement vsync check box in UFE
TTP331451
#ue4
#ufe

diff --git a/Engine/Source/Developer/LauncherServices/Private/Launcher/LauncherDeployCommands.h b/Engine/Source/Developer/LauncherServices/Private/Launcher/LauncherDeployCommands.h
--- a/Engine/Source/Developer/LauncherServices/Private/Launcher/LauncherDeployCommands.h
+++ b/Engine/Source/Developer/LauncherServices/Private/Launcher/LauncherDeployCommands.h
@@ -63,16 +63,17 @@ public:
 		CommandLine += FString::Printf(TEXT(" -cmdline=\"%s -Messaging\""),
 			*InitialMap);
 
-		CommandLine += FString::Printf(TEXT(" -addcmdline=\"%s -InstanceId=%s -SessionId=%s -SessionOwner=%s -SessionName='%s' %s %s %s\""),
+		CommandLine += FString::Printf(TEXT(" -addcmdline=\"%s -InstanceId=%s -SessionId=%s -SessionOwner=%s -SessionName='%s'%s%s%s%s\""),
 			*InitialMap,
 			*InstanceId.ToString(),
 			*ChainState.SessionId.ToString(),
 			FPlatformProcess::UserName(false),
 			*ChainState.Profile->GetName(),
-			CookCommand.IsValid() ? *CookCommand->GetAdditionalArguments(ChainState) : TEXT(""),
+			CookCommand.IsValid() ? *(TEXT(" ") + CookCommand->GetAdditionalArguments(ChainState)) : TEXT(""),
 			((TargetPlatform.PlatformName() == TEXT("PS4") || ChainState.Profile->IsPackingWithUnrealPak())
-			&& ChainState.Profile->GetCookMode() == ELauncherProfileCookModes::ByTheBook) ? TEXT("-pak") : TEXT(""),
-			*LauncherCommandLine);
+			&& ChainState.Profile->GetCookMode() == ELauncherProfileCookModes::ByTheBook) ? TEXT(" -pak") : TEXT(""),
+			ChainState.Profile->GetLaunchRoles().Num() > 0 ? (ChainState.Profile->GetLaunchRoles()[0]->IsVsyncEnabled() ? TEXT(" -vsync") : TEXT("")) : TEXT(""),
+			*(TEXT(" ") + LauncherCommandLine));
 
 		return CommandLine;
 	}
@@ -177,15 +178,16 @@ public:
 		CommandLine += FString::Printf(TEXT(" -cmdline=\"%s -Messaging\""),
 			*InitialMap);
 
-		CommandLine += FString::Printf(TEXT(" -addcmdline=\"%s -InstanceId=%s -SessionId=%s -SessionOwner=%s -SessionName='%s' %s %s\""),
+		CommandLine += FString::Printf(TEXT(" -addcmdline=\"%s -InstanceId=%s -SessionId=%s -SessionOwner=%s -SessionName='%s'%s%s%s\""),
 			*InitialMap,
 			*InstanceId.ToString(),
 			*ChainState.SessionId.ToString(),
 			FPlatformProcess::UserName(false),
 			*ChainState.Profile->GetName(),
-			CookCommand.IsValid() ? *CookCommand->GetAdditionalArguments(ChainState) : TEXT(""),
+			CookCommand.IsValid() ? *(TEXT(" ") + CookCommand->GetAdditionalArguments(ChainState)) : TEXT(""),
 			((TargetPlatform.PlatformName() == TEXT("PS4") || ChainState.Profile->IsPackingWithUnrealPak())
-			&& ChainState.Profile->GetCookMode() == ELauncherProfileCookModes::ByTheBook) ? TEXT("-pak") : TEXT(""));
+			&& ChainState.Profile->GetCookMode() == ELauncherProfileCookModes::ByTheBook) ? TEXT(" -pak") : TEXT(""),
+			ChainState.Profile->GetLaunchRoles().Num() > 0 ? (ChainState.Profile->GetLaunchRoles()[0]->IsVsyncEnabled() ? TEXT(" -vsync") : TEXT("")) : TEXT(""));
 
 		return CommandLine;
 	}
@@ -255,16 +257,17 @@ public:
 		CommandLine += FString::Printf(TEXT(" -cmdline=\"%s -Messaging\""),
 			*InitialMap);
 
-		CommandLine += FString::Printf(TEXT(" -addcmdline=\"%s -InstanceId=%s -SessionId=%s -SessionOwner=%s -SessionName='%s' %s %s %s\""),
+		CommandLine += FString::Printf(TEXT(" -addcmdline=\"%s -InstanceId=%s -SessionId=%s -SessionOwner=%s -SessionName='%s'%s%s%s%s\""),
 			*InitialMap,
 			*InstanceId.ToString(),
 			*ChainState.SessionId.ToString(),
 			FPlatformProcess::UserName(false),
 			*ChainState.Profile->GetName(),
-			CookCommand.IsValid() ? *CookCommand->GetAdditionalArguments(ChainState) : TEXT(""),
+			CookCommand.IsValid() ? *(TEXT(" ") + CookCommand->GetAdditionalArguments(ChainState)) : TEXT(""),
 			((TargetPlatform.PlatformName() == TEXT("PS4") || ChainState.Profile->IsPackingWithUnrealPak())
-			&& ChainState.Profile->GetCookMode() == ELauncherProfileCookModes::ByTheBook) ? TEXT("-pak") : TEXT(""),
-			*LauncherCommandLine);
+			&& ChainState.Profile->GetCookMode() == ELauncherProfileCookModes::ByTheBook) ? TEXT(" -pak") : TEXT(""),
+			ChainState.Profile->GetLaunchRoles().Num() > 0 ? (ChainState.Profile->GetLaunchRoles()[0]->IsVsyncEnabled() ? TEXT(" -vsync") : TEXT("")) : TEXT(""),
+			*(TEXT(" ") + LauncherCommandLine));
 		return CommandLine;
 	}
 
@@ -365,15 +368,16 @@ public:
 		CommandLine += FString::Printf(TEXT(" -cmdline=\"%s -Messaging\""),
 			*InitialMap);
 
-		CommandLine += FString::Printf(TEXT(" -addcmdline=\"%s -InstanceId=%s -SessionId=%s -SessionOwner=%s -SessionName='%s' %s %s\""),
+		CommandLine += FString::Printf(TEXT(" -addcmdline=\"%s -InstanceId=%s -SessionId=%s -SessionOwner=%s -SessionName='%s'%s%s%s\""),
 			*InitialMap,
 			*InstanceId.ToString(),
 			*ChainState.SessionId.ToString(),
 			FPlatformProcess::UserName(false),
 			*ChainState.Profile->GetName(),
-			CookCommand.IsValid() ? *CookCommand->GetAdditionalArguments(ChainState) : TEXT(""),
+			CookCommand.IsValid() ? *(TEXT(" ") + CookCommand->GetAdditionalArguments(ChainState)) : TEXT(""),
 			((TargetPlatform.PlatformName() == TEXT("PS4") || ChainState.Profile->IsPackingWithUnrealPak())
-			&& ChainState.Profile->GetCookMode() == ELauncherProfileCookModes::ByTheBook) ? TEXT("-pak") : TEXT(""));
+			&& ChainState.Profile->GetCookMode() == ELauncherProfileCookModes::ByTheBook) ? TEXT(" -pak") : TEXT(""),
+			ChainState.Profile->GetLaunchRoles().Num() > 0 ? (ChainState.Profile->GetLaunchRoles()[0]->IsVsyncEnabled() ? TEXT(" -vsync") : TEXT("")) : TEXT(""));
 
 		return CommandLine;
 	}
diff --git a/Engine/Source/Developer/LauncherServices/Private/Launcher/LauncherLaunchCommands.h b/Engine/Source/Developer/LauncherServices/Private/Launcher/LauncherLaunchCommands.h
--- a/Engine/Source/Developer/LauncherServices/Private/Launcher/LauncherLaunchCommands.h
+++ b/Engine/Source/Developer/LauncherServices/Private/Launcher/LauncherLaunchCommands.h
@@ -62,13 +62,14 @@ public:
 		CommandLine += CookCommand.IsValid() ? CookCommand->GetDependencyArguments(ChainState) : TEXT(" -skipcook");
 
 		// additional command line arguments
-		CommandLine += FString::Printf(TEXT(" -addcmdline=\"-InstanceId=%s -SessionId=%s -SessionOwner=%s -SessionName='%s' %s %s\""),
+		CommandLine += FString::Printf(TEXT(" -addcmdline=\"-InstanceId=%s -SessionId=%s -SessionOwner=%s -SessionName='%s'%s%s%s\""),
 			*InstanceId.ToString(),
 			*ChainState.SessionId.ToString(),
 			FPlatformProcess::UserName(false),
 			*ChainState.Profile->GetName(),
-			CookCommand.IsValid() ? *CookCommand->GetAdditionalArguments(ChainState) : TEXT(""),
-			*Role->GetCommandLine()
+			CookCommand.IsValid() ? *(TEXT(" ") + CookCommand->GetAdditionalArguments(ChainState)) : TEXT(""),
+			Role->IsVsyncEnabled() ? TEXT(" -vsync") : TEXT(""),
+			*(TEXT(" ") + Role->GetCommandLine())
 			);
 
 		return CommandLine;
@@ -161,12 +162,12 @@ public:
 		CommandLine += CookCommand.IsValid() ? CookCommand->GetDependencyArguments(ChainState) : TEXT("");
 
 		// additional command line arguments
-		CommandLine += FString::Printf(TEXT(" -addcmdline=\"-InstanceId=%s -SessionId=%s -SessionOwner=%s -SessionName='%s' %s %s\""),
+		CommandLine += FString::Printf(TEXT(" -addcmdline=\"-InstanceId=%s -SessionId=%s -SessionOwner=%s -SessionName='%s'%s%s\""),
 			*InstanceId.ToString(),
 			*ChainState.SessionId.ToString(),
 			FPlatformProcess::UserName(false),
 			*ChainState.Profile->GetName(),
-			CookCommand.IsValid() ? *CookCommand->GetAdditionalArguments(ChainState) : TEXT(""),
+			CookCommand.IsValid() ? *(TEXT(" ") + CookCommand->GetAdditionalArguments(ChainState)) : TEXT(""),
 			*Role->GetCommandLine()
 			);
 
diff --git a/Engine/Source/Developer/LauncherServices/Private/Launcher/LauncherPackageCommands.h b/Engine/Source/Developer/LauncherServices/Private/Launcher/LauncherPackageCommands.h
--- a/Engine/Source/Developer/LauncherServices/Private/Launcher/LauncherPackageCommands.h
+++ b/Engine/Source/Developer/LauncherServices/Private/Launcher/LauncherPackageCommands.h
@@ -53,15 +53,17 @@ public:
         CommandLine += FString::Printf(TEXT(" -cmdline=\"%s -Messaging\""),
 										*InitialMap);
         
-		CommandLine += FString::Printf(TEXT(" -addcmdline=\"%s -InstanceId=%s -SessionId=%s -SessionOwner=%s -SessionName='%s' %s %s\""),
+		CommandLine += FString::Printf(TEXT(" -addcmdline=\"%s -InstanceId=%s -SessionId=%s -SessionOwner=%s -SessionName='%s'%s%s%s\""),
 			*InitialMap,
 			*InstanceId.ToString(),
 			*ChainState.SessionId.ToString(),
 			FPlatformProcess::UserName(false),
 			*ChainState.Profile->GetName(),
 			CookCommand.IsValid() ? *CookCommand->GetAdditionalArguments(ChainState) : TEXT(""),
+			CookCommand.IsValid() ? *(TEXT(" ") + CookCommand->GetAdditionalArguments(ChainState)) : TEXT(""),
 			((TargetPlatform.PlatformName() == TEXT("PS4") || ChainState.Profile->IsPackingWithUnrealPak())
-			&& ChainState.Profile->GetCookMode() == ELauncherProfileCookModes::ByTheBook) ? TEXT("-pak") : TEXT(""));
+			&& ChainState.Profile->GetCookMode() == ELauncherProfileCookModes::ByTheBook) ? TEXT(" -pak") : TEXT(""),
+			ChainState.Profile->GetLaunchRoles().Num() > 0 ? (ChainState.Profile->GetLaunchRoles()[0]->IsVsyncEnabled() ? TEXT(" -vsync") : TEXT("")) : TEXT(""));
 
 		return CommandLine;
 	}
@@ -142,15 +144,16 @@ public:
 		CommandLine += FString::Printf(TEXT(" -cmdline=\"%s -Messaging\""),
 			*InitialMap);
 
-		CommandLine += FString::Printf(TEXT(" -addcmdline=\"%s -InstanceId=%s -SessionId=%s -SessionOwner=%s -SessionName='%s' %s %s\""),
+		CommandLine += FString::Printf(TEXT(" -addcmdline=\"%s -InstanceId=%s -SessionId=%s -SessionOwner=%s -SessionName='%s'%s%s%s\""),
 			*InitialMap,
 			*InstanceId.ToString(),
 			*ChainState.SessionId.ToString(),
 			FPlatformProcess::UserName(false),
 			*ChainState.Profile->GetName(),
-			CookCommand.IsValid() ? *CookCommand->GetAdditionalArguments(ChainState) : TEXT(""),
+			CookCommand.IsValid() ? *(TEXT(" ") + CookCommand->GetAdditionalArguments(ChainState)) : TEXT(""),
 			((TargetPlatform.PlatformName() == TEXT("PS4") || ChainState.Profile->IsPackingWithUnrealPak())
-			&& ChainState.Profile->GetCookMode() == ELauncherProfileCookModes::ByTheBook) ? TEXT("-pak") : TEXT(""));
+			&& ChainState.Profile->GetCookMode() == ELauncherProfileCookModes::ByTheBook) ? TEXT(" -pak") : TEXT(""),
+			ChainState.Profile->GetLaunchRoles().Num() > 0 ? (ChainState.Profile->GetLaunchRoles()[0]->IsVsyncEnabled() ? TEXT(" -vsync") : TEXT("")) : TEXT(""));
 
 		return CommandLine;
 	}
@@ -216,15 +219,16 @@ public:
 		CommandLine += FString::Printf(TEXT(" -cmdline=\"%s -Messaging\""),
 			*InitialMap);
 
-		CommandLine += FString::Printf(TEXT(" -addcmdline=\"%s -InstanceId=%s -SessionId=%s -SessionOwner=%s -SessionName='%s' %s %s\""),
+		CommandLine += FString::Printf(TEXT(" -addcmdline=\"%s -InstanceId=%s -SessionId=%s -SessionOwner=%s -SessionName='%s'%s%s%s\""),
 			*InitialMap,
 			*InstanceId.ToString(),
 			*ChainState.SessionId.ToString(),
 			FPlatformProcess::UserName(false),
 			*ChainState.Profile->GetName(),
-			CookCommand.IsValid() ? *CookCommand->GetAdditionalArguments(ChainState) : TEXT(""),
+			CookCommand.IsValid() ? *(TEXT(" ") + CookCommand->GetAdditionalArguments(ChainState)) : TEXT(""),
 			((TargetPlatform.PlatformName() == TEXT("PS4") || ChainState.Profile->IsPackingWithUnrealPak())
-			&& ChainState.Profile->GetCookMode() == ELauncherProfileCookModes::ByTheBook) ? TEXT("-pak") : TEXT(""));
+			&& ChainState.Profile->GetCookMode() == ELauncherProfileCookModes::ByTheBook) ? TEXT(" -pak") : TEXT(""),
+			ChainState.Profile->GetLaunchRoles().Num() > 0 ? (ChainState.Profile->GetLaunchRoles()[0]->IsVsyncEnabled() ? TEXT(" -vsync") : TEXT("")) : TEXT(""));
 
 		return CommandLine;
 	}
@@ -320,15 +324,16 @@ public:
 		CommandLine += FString::Printf(TEXT(" -cmdline=\"%s -Messaging\""),
 			*InitialMap);
 
-		CommandLine += FString::Printf(TEXT(" -addcmdline=\"%s -InstanceId=%s -SessionId=%s -SessionOwner=%s -SessionName='%s' %s %s\""),
+		CommandLine += FString::Printf(TEXT(" -addcmdline=\"%s -InstanceId=%s -SessionId=%s -SessionOwner=%s -SessionName='%s'%s%s%s\""),
 			*InitialMap,
 			*InstanceId.ToString(),
 			*ChainState.SessionId.ToString(),
 			FPlatformProcess::UserName(false),
 			*ChainState.Profile->GetName(),
-			CookCommand.IsValid() ? *CookCommand->GetAdditionalArguments(ChainState) : TEXT(""),
+			CookCommand.IsValid() ? *(TEXT(" ") + CookCommand->GetAdditionalArguments(ChainState)) : TEXT(""),
 			((TargetPlatform.PlatformName() == TEXT("PS4") || ChainState.Profile->IsPackingWithUnrealPak())
-			&& ChainState.Profile->GetCookMode() == ELauncherProfileCookModes::ByTheBook) ? TEXT("-pak") : TEXT(""));
+			&& ChainState.Profile->GetCookMode() == ELauncherProfileCookModes::ByTheBook) ? TEXT(" -pak") : TEXT(""),
+			ChainState.Profile->GetLaunchRoles().Num() > 0 ? (ChainState.Profile->GetLaunchRoles()[0]->IsVsyncEnabled() ? TEXT(" -vsync") : TEXT("")) : TEXT(""));
 
 		return CommandLine;
 	}
diff --git a/Engine/Source/Runtime/Engine/Private/UnrealEngine.cpp b/Engine/Source/Runtime/Engine/Private/UnrealEngine.cpp
--- a/Engine/Source/Runtime/Engine/Private/UnrealEngine.cpp
+++ b/Engine/Source/Runtime/Engine/Private/UnrealEngine.cpp
@@ -705,6 +705,17 @@ void UEngine::Init(IEngineLoop* InEngineLoop)
 		}
 	}
 
+	// optionally set the vsync console variable
+	if( FParse::Param(FCommandLine::Get(), TEXT("vsync")) )
+	{
+		new(GEngine->DeferredCommands) FString(TEXT("r.vsync 1"));
+	}
+
+	// optionally set the vsync console variable
+	if( FParse::Param(FCommandLine::Get(), TEXT("novsync")) )
+	{
+		new(GEngine->DeferredCommands) FString(TEXT("r.vsync 0"));
+	}
 #endif // !(UE_BUILD_SHIPPING || UE_BUILD_TEST)
 
 	if (GetDerivedDataCache())
