quick optimization to not re-build UAT for every step of the launcher process
#ue4

diff --git a/Engine/Source/Developer/LauncherServices/Private/Launcher/LauncherUATTask.h b/Engine/Source/Developer/LauncherServices/Private/Launcher/LauncherUATTask.h
--- a/Engine/Source/Developer/LauncherServices/Private/Launcher/LauncherUATTask.h
+++ b/Engine/Source/Developer/LauncherServices/Private/Launcher/LauncherUATTask.h
@@ -27,9 +27,10 @@ public:
 		, TaskCommand(InCommand)
 		, TargetPlatform(InTargetPlatform)
 	{
-		NoCompile = !FParse::Param( FCommandLine::Get(), TEXT("development") ) ? TEXT(" -nocompile") : TEXT("");
+		NoCompile = TEXT(" -nocompile");
 	}
 
+	static bool FirstTimeCompile;
 protected:
 
 	/**
@@ -41,6 +42,12 @@ protected:
 	 */
 	virtual bool PerformTask( FLauncherTaskChainState& ChainState ) override
 	{
+		if (FirstTimeCompile)
+		{
+			NoCompile = !FParse::Param( FCommandLine::Get(), TEXT("development") ) ? TEXT(" -nocompile") : TEXT("");
+			FirstTimeCompile = false;
+		}
+
 		// spawn a UAT process to cook the data
 		// UAT executable
 		FString ExecutablePath = FPaths::ConvertRelativePathToFull(FPaths::EngineDir() + FString(TEXT("Build")) / TEXT("BatchFiles"));
diff --git a/Engine/Source/Developer/LauncherServices/Private/Launcher/LauncherWorker.cpp b/Engine/Source/Developer/LauncherServices/Private/Launcher/LauncherWorker.cpp
--- a/Engine/Source/Developer/LauncherServices/Private/Launcher/LauncherWorker.cpp
+++ b/Engine/Source/Developer/LauncherServices/Private/Launcher/LauncherWorker.cpp
@@ -3,6 +3,8 @@
 #include "LauncherServicesPrivatePCH.h"
 
 
+bool FLauncherUATTask::FirstTimeCompile = true;
+
 /* FLauncherWorker structors
  *****************************************************************************/
 
@@ -197,6 +199,8 @@ void FLauncherWorker::CreateAndExecuteTasks( const ILauncherProfileRef& InProfil
 	TSharedPtr<FLauncherTask> FirstPlatformPackageTask = NULL;
 	TSharedPtr<FLauncherTask> FirstPlatformDeviceTask = NULL;
 
+	FLauncherUATTask::FirstTimeCompile = true;
+
 	// determine deployment platforms
 	ILauncherDeviceGroupPtr DeviceGroup = InProfile->GetDeployedDeviceGroup();
 
