Fix PS4/XB1 Launching getting stuck at deploy step

diff --git a/Engine/Source/Developer/LauncherServices/Private/Launcher/LauncherDeployCommands.h b/Engine/Source/Developer/LauncherServices/Private/Launcher/LauncherDeployCommands.h
--- a/Engine/Source/Developer/LauncherServices/Private/Launcher/LauncherDeployCommands.h
+++ b/Engine/Source/Developer/LauncherServices/Private/Launcher/LauncherDeployCommands.h
@@ -55,9 +55,13 @@ public:
 			*StagePath,
 			*InitialMap,
 			*TargetPlatform.PlatformName());
-#if !PLATFORM_WINDOWS
-		CommandLine += FString::Printf(TEXT(" -device=\"%s\""), *DeviceProxy->GetDeviceId());
+		
+#if !PLATFORM_MAC
+		if (TargetPlatform.PlatformName() != TEXT("IOS"))
 #endif
+		{
+			CommandLine += FString::Printf(TEXT(" -device=\"%s\""), *DeviceProxy->GetDeviceId());
+		}
 
 		// cook dependency arguments
 		CommandLine += CookCommand.IsValid() ? CookCommand->GetDependencyArguments(ChainState) : TEXT(" -skipcook");
