fixed launcher to properly require HTML5 to package whenever a deploy is requested
#ue4
#html5

diff --git a/Engine/Source/Developer/LauncherServices/Private/Launcher/LauncherWorker.cpp b/Engine/Source/Developer/LauncherServices/Private/Launcher/LauncherWorker.cpp
--- a/Engine/Source/Developer/LauncherServices/Private/Launcher/LauncherWorker.cpp
+++ b/Engine/Source/Developer/LauncherServices/Private/Launcher/LauncherWorker.cpp
@@ -310,7 +310,7 @@ void FLauncherWorker::CreateAndExecuteTasks( const ILauncherProfileRef& InProfil
         }
 
 		// ... package the build...
-		if (InProfile->GetPackagingMode() != ELauncherProfilePackagingModes::DoNotPackage || (TargetPlatform->PlatformName() == TEXT("IOS") && InProfile->GetDeploymentMode() != ELauncherProfileDeploymentModes::CopyRepository && InProfile->GetDeploymentMode() != ELauncherProfileDeploymentModes::DoNotDeploy))
+		if (InProfile->GetPackagingMode() != ELauncherProfilePackagingModes::DoNotPackage || ((TargetPlatform->PlatformName() == TEXT("IOS") || TargetPlatform->PlatformName() == TEXT("HTML5")) && InProfile->GetDeploymentMode() != ELauncherProfileDeploymentModes::CopyRepository && InProfile->GetDeploymentMode() != ELauncherProfileDeploymentModes::DoNotDeploy))
 		{
 			TSharedPtr<FLauncherUATCommand> Command = NULL;
 			if (TargetPlatform->PlatformName() == TEXT("WindowsServer") || TargetPlatform->PlatformName() == TEXT("LinuxServer"))
