Allow setting the path to the editor executable when using launcher commands.

diff --git a/Engine/Source/Developer/LauncherServices/Private/Launcher/LauncherUATTask.h b/Engine/Source/Developer/LauncherServices/Private/Launcher/LauncherUATTask.h
--- a/Engine/Source/Developer/LauncherServices/Private/Launcher/LauncherUATTask.h
+++ b/Engine/Source/Developer/LauncherServices/Private/Launcher/LauncherUATTask.h
@@ -22,10 +22,11 @@ public:
 	 * @param InTargetPlatform - the target platform for the task
 	 * @param InName - The name of the task.
 	 */
-	FLauncherUATTask( const TSharedPtr<FLauncherUATCommand>& InCommand, const ITargetPlatform& InTargetPlatform, const FString& InName, void* InReadPipe, void* InWritePipe)
+	FLauncherUATTask( const TSharedPtr<FLauncherUATCommand>& InCommand, const ITargetPlatform& InTargetPlatform, const FString& InName, void* InReadPipe, void* InWritePipe, const FString& InEditorExe )
 		: FLauncherTask(InName, InCommand->GetDesc(), InReadPipe, InWritePipe)
 		, TaskCommand(InCommand)
 		, TargetPlatform(InTargetPlatform)
+		, EditorExe(InEditorExe)
 	{
 		NoCompile = TEXT(" -nocompile");
 	}
@@ -105,6 +106,12 @@ protected:
         CommandLine += Rocket;
 		CommandLine += OptionalParams;
 
+		// specify the path to the editor exe if necessary
+		if(EditorExe.Len() > 0)
+		{
+			CommandLine += FString::Printf(TEXT(" -ue4exe=\"%s\""), *EditorExe);
+		}
+
 		// specialized command arguments for this particular task
 		CommandLine += TaskCommand->GetArguments(ChainState);
 
@@ -152,4 +159,7 @@ private:
 
 	// Holds the no compile flag
 	FString NoCompile;
+
+	// The editor executable that UAT should use
+	FString EditorExe;
 };
diff --git a/Engine/Source/Developer/LauncherServices/Private/Launcher/LauncherWorker.cpp b/Engine/Source/Developer/LauncherServices/Private/Launcher/LauncherWorker.cpp
--- a/Engine/Source/Developer/LauncherServices/Private/Launcher/LauncherWorker.cpp
+++ b/Engine/Source/Developer/LauncherServices/Private/Launcher/LauncherWorker.cpp
@@ -245,7 +245,7 @@ void FLauncherWorker::CreateAndExecuteTasks( const ILauncherProfileRef& InProfil
 			{
 				Command = MakeShareable(new FLauncherBuildGameCommand(*TargetPlatform));
 			}
-			TSharedPtr<FLauncherTask> BuildTask = MakeShareable(new FLauncherUATTask(Command, *TargetPlatform, Command->GetName(), ReadPipe, WritePipe));
+			TSharedPtr<FLauncherTask> BuildTask = MakeShareable(new FLauncherUATTask(Command, *TargetPlatform, Command->GetName(), ReadPipe, WritePipe, InProfile->GetEditorExe()));
 			BuildTask->OnStarted().AddRaw(this, &FLauncherWorker::OnTaskStarted);
 			BuildTask->OnCompleted().AddRaw(this, &FLauncherWorker::OnTaskCompleted);
 			if (!PerPlatformBuildTask.IsValid())
@@ -275,7 +275,7 @@ void FLauncherWorker::CreateAndExecuteTasks( const ILauncherProfileRef& InProfil
 				Command = MakeShareable(new FLauncherCookGameCommand(*TargetPlatform));
 			}
 			CookCommand = Command;
-			TSharedPtr<FLauncherTask> CookTask = MakeShareable(new FLauncherUATTask(Command, *TargetPlatform, Command->GetName(), ReadPipe, WritePipe));
+			TSharedPtr<FLauncherTask> CookTask = MakeShareable(new FLauncherUATTask(Command, *TargetPlatform, Command->GetName(), ReadPipe, WritePipe, InProfile->GetEditorExe()));
 			CookTask->OnStarted().AddRaw(this, &FLauncherWorker::OnTaskStarted);
 			CookTask->OnCompleted().AddRaw(this, &FLauncherWorker::OnTaskCompleted);
 			if (!PerPlatformCookTask.IsValid())
@@ -322,7 +322,7 @@ void FLauncherWorker::CreateAndExecuteTasks( const ILauncherProfileRef& InProfil
 				Command = MakeShareable(new FLauncherPackageGameCommand(*TargetPlatform, CookCommand));
 			}
 			
-			TSharedPtr<FLauncherTask> PackageTask = MakeShareable(new FLauncherUATTask(Command, *TargetPlatform, Command->GetName(), ReadPipe, WritePipe));
+			TSharedPtr<FLauncherTask> PackageTask = MakeShareable(new FLauncherUATTask(Command, *TargetPlatform, Command->GetName(), ReadPipe, WritePipe, InProfile->GetEditorExe()));
 			PackageTask->OnStarted().AddRaw(this, &FLauncherWorker::OnTaskStarted);
 			PackageTask->OnCompleted().AddRaw(this, &FLauncherWorker::OnTaskCompleted);
 			if (!PerPlatformPackageTask.IsValid())
@@ -348,7 +348,7 @@ void FLauncherWorker::CreateAndExecuteTasks( const ILauncherProfileRef& InProfil
 				Command = MakeShareable(new FLauncherStageGameCommand(*TargetPlatform, CookCommand));
 			}
 
-			TSharedPtr<FLauncherTask> StageTask = MakeShareable(new FLauncherUATTask(Command, *TargetPlatform, Command->GetName(), ReadPipe, WritePipe));
+			TSharedPtr<FLauncherTask> StageTask = MakeShareable(new FLauncherUATTask(Command, *TargetPlatform, Command->GetName(), ReadPipe, WritePipe, InProfile->GetEditorExe()));
 			StageTask->OnStarted().AddRaw(this, &FLauncherWorker::OnTaskStarted);
 			StageTask->OnCompleted().AddRaw(this, &FLauncherWorker::OnTaskCompleted);
 			if (!PerPlatformPackageTask.IsValid())
@@ -408,7 +408,7 @@ void FLauncherWorker::CreateAndExecuteTasks( const ILauncherProfileRef& InProfil
 							Command = MakeShareable(new FLauncherCookOnTheFlyCommand(*TargetPlatform));
 						}
 						CookCommand = Command;
-						FileServerTask = MakeShareable(new FLauncherUATTask(Command, *TargetPlatform, Command->GetName(), ReadPipe, WritePipe));
+						FileServerTask = MakeShareable(new FLauncherUATTask(Command, *TargetPlatform, Command->GetName(), ReadPipe, WritePipe, InProfile->GetEditorExe()));
 						FileServerTask->OnStarted().AddRaw(this, &FLauncherWorker::OnTaskStarted);
 						FileServerTask->OnCompleted().AddRaw(this, &FLauncherWorker::OnTaskCompleted);
 					}
@@ -437,7 +437,7 @@ void FLauncherWorker::CreateAndExecuteTasks( const ILauncherProfileRef& InProfil
 					{
 						Command = MakeShareable(new FLauncherDeployGameToDeviceCommand(DeviceProxy, *TargetPlatform, CookCommand, LaunchCommandLine));
 					}
-					TSharedPtr<FLauncherTask> DeployTask = MakeShareable(new FLauncherUATTask(Command, *TargetPlatform, Command->GetName(), ReadPipe, WritePipe));
+					TSharedPtr<FLauncherTask> DeployTask = MakeShareable(new FLauncherUATTask(Command, *TargetPlatform, Command->GetName(), ReadPipe, WritePipe, InProfile->GetEditorExe()));
 					DeployTask->OnStarted().AddRaw(this, &FLauncherWorker::OnTaskStarted);
 					DeployTask->OnCompleted().AddRaw(this, &FLauncherWorker::OnTaskCompleted);
 
@@ -463,7 +463,7 @@ void FLauncherWorker::CreateAndExecuteTasks( const ILauncherProfileRef& InProfil
 					{
 						Command = MakeShareable(new FLauncherDeployGamePackageToDeviceCommand(DeviceProxy, *TargetPlatform, CookCommand, LaunchCommandLine));
 					}
-					TSharedPtr<FLauncherTask> DeployTask = MakeShareable(new FLauncherUATTask(Command, *TargetPlatform, Command->GetName(), ReadPipe, WritePipe));
+					TSharedPtr<FLauncherTask> DeployTask = MakeShareable(new FLauncherUATTask(Command, *TargetPlatform, Command->GetName(), ReadPipe, WritePipe, InProfile->GetEditorExe()));
 					DeployTask->OnStarted().AddRaw(this, &FLauncherWorker::OnTaskStarted);
 					DeployTask->OnCompleted().AddRaw(this, &FLauncherWorker::OnTaskCompleted);
 
@@ -481,7 +481,7 @@ void FLauncherWorker::CreateAndExecuteTasks( const ILauncherProfileRef& InProfil
 				else if (TargetPlatform->PlatformName() == TEXT("XboxOne") || TargetPlatform->PlatformName() == TEXT("IOS") || TargetPlatform->PlatformName().StartsWith(TEXT("Android")))
 				{
 					TSharedPtr<FLauncherUATCommand> Command = MakeShareable(new FLauncherDeployGameToDeviceCommand(DeviceProxy, *TargetPlatform, CookCommand, LaunchCommandLine));
-					TSharedPtr<FLauncherTask> DeployTask = MakeShareable(new FLauncherUATTask(Command, *TargetPlatform, Command->GetName(), ReadPipe, WritePipe));
+					TSharedPtr<FLauncherTask> DeployTask = MakeShareable(new FLauncherUATTask(Command, *TargetPlatform, Command->GetName(), ReadPipe, WritePipe, InProfile->GetEditorExe()));
 					DeployTask->OnStarted().AddRaw(this, &FLauncherWorker::OnTaskStarted);
 					DeployTask->OnCompleted().AddRaw(this, &FLauncherWorker::OnTaskCompleted);
 
@@ -517,7 +517,7 @@ void FLauncherWorker::CreateAndExecuteTasks( const ILauncherProfileRef& InProfil
 								Command = MakeShareable(new FLauncherLaunchDedicatedServerCommand(DeviceProxy, *TargetPlatform, Roles[RoleIndex].ToSharedRef(), CookCommand));
 							}
 								
-							TSharedPtr<FLauncherTask> PerRoleTask = MakeShareable(new FLauncherUATTask(Command, *TargetPlatform, Command->GetName(), ReadPipe, WritePipe));
+							TSharedPtr<FLauncherTask> PerRoleTask = MakeShareable(new FLauncherUATTask(Command, *TargetPlatform, Command->GetName(), ReadPipe, WritePipe, InProfile->GetEditorExe()));
 							PerRoleTask->OnStarted().AddRaw(this, &FLauncherWorker::OnTaskStarted);
 							PerRoleTask->OnCompleted().AddRaw(this, &FLauncherWorker::OnTaskCompleted);
 
@@ -544,7 +544,7 @@ void FLauncherWorker::CreateAndExecuteTasks( const ILauncherProfileRef& InProfil
 			if (!FileServerTask.IsValid())
 			{
 				TSharedPtr<FLauncherUATCommand> Command = MakeShareable(new FLauncherStandAloneCookOnTheFlyCommand(*TargetPlatform));
-				FileServerTask = MakeShareable(new FLauncherUATTask(Command, *TargetPlatform, Command->GetName(), ReadPipe, WritePipe));
+				FileServerTask = MakeShareable(new FLauncherUATTask(Command, *TargetPlatform, Command->GetName(), ReadPipe, WritePipe, InProfile->GetEditorExe()));
 				FileServerTask->OnStarted().AddRaw(this, &FLauncherWorker::OnTaskStarted);
 				FileServerTask->OnCompleted().AddRaw(this, &FLauncherWorker::OnTaskCompleted);
 			}
diff --git a/Engine/Source/Developer/LauncherServices/Private/Profiles/LauncherProfile.h b/Engine/Source/Developer/LauncherServices/Private/Profiles/LauncherProfile.h
--- a/Engine/Source/Developer/LauncherServices/Private/Profiles/LauncherProfile.h
+++ b/Engine/Source/Developer/LauncherServices/Private/Profiles/LauncherProfile.h
@@ -481,6 +481,9 @@ public:
 		// default packaging settings
 		PackagingMode = ELauncherProfilePackagingModes::DoNotPackage;
 
+		// default UAT settings
+		EditorExe.Empty();
+
 		Validate();
 	}
 
@@ -683,6 +686,16 @@ public:
 		return ProjectChangedDelegate;
 	}
 
+	virtual void SetEditorExe( const FString& InEditorExe ) override
+	{
+		EditorExe = InEditorExe;
+	}
+
+	virtual FString GetEditorExe( ) const override
+	{
+		return EditorExe;
+	}
+
 	// End ILauncherProfile interface
 
 protected:
@@ -931,6 +944,9 @@ private:
 
     // Holds the close value for the cook on the fly server
     bool ForceClose;
+
+	// Path to the editor executable to pass to UAT, for cooking, etc... May be empty.
+	FString EditorExe;
         
 private:
 
diff --git a/Engine/Source/Developer/LauncherServices/Public/Interfaces/ILauncherProfile.h b/Engine/Source/Developer/LauncherServices/Public/Interfaces/ILauncherProfile.h
--- a/Engine/Source/Developer/LauncherServices/Public/Interfaces/ILauncherProfile.h
+++ b/Engine/Source/Developer/LauncherServices/Public/Interfaces/ILauncherProfile.h
@@ -479,6 +479,20 @@ public:
 	 */
 	virtual bool SupportsEngineMaps( ) const = 0;
 
+	/**
+	 * Sets the path to the editor executable to use in UAT.
+	 *
+	 * @param EditorExe Path to the editor executable.
+	 */
+	virtual void SetEditorExe( const FString& EditorExe ) = 0;
+
+	/**
+	 * Gets the path to the editor executable.
+	 *
+	 * @return Path to the editor executable.
+	 */
+	virtual FString GetEditorExe( ) const = 0;
+
 public:
 
 	/**
diff --git a/Engine/Source/Editor/MainFrame/Private/Frame/MainFrameActions.cpp b/Engine/Source/Editor/MainFrame/Private/Frame/MainFrameActions.cpp
--- a/Engine/Source/Editor/MainFrame/Private/Frame/MainFrameActions.cpp
+++ b/Engine/Source/Editor/MainFrame/Private/Frame/MainFrameActions.cpp
@@ -429,19 +429,6 @@ void FMainFrameActionCallbacks::CookContent( const FName InPlatformInfoName )
 		OptionalParams += PlatformInfo->UATCommandLine;
 	}
 
-	FString ExecutableName = FPlatformProcess::ExecutableName(false);
-
-#if PLATFORM_WINDOWS
-	// turn UE4editor into UE4editor-cmd
-	if(ExecutableName.EndsWith(".exe", ESearchCase::IgnoreCase) && !FPaths::GetBaseFilename(ExecutableName).EndsWith("-cmd", ESearchCase::IgnoreCase))
-	{
-		FString NewExeName = ExecutableName.Left(ExecutableName.Len() - 4) + "-Cmd.exe";
-		if (FPaths::FileExists(NewExeName))
-		{
-			ExecutableName = NewExeName;
-		}
-	}
-#endif
 	const bool bRunningDebug = FParse::Param(FCommandLine::Get(), TEXT("debug"));
 
 	if (bRunningDebug)
@@ -455,7 +442,7 @@ void FMainFrameActionCallbacks::CookContent( const FName InPlatformInfoName )
 		FApp::IsEngineInstalled() ? TEXT(" -installed") : TEXT(""),
 		*ProjectPath,
 		*PlatformInfo->TargetPlatformName.ToString(),
-		*ExecutableName,
+		*FUnrealEdMisc::Get().GetExecutableForCommandlets(),
 		*OptionalParams
 	);
 
diff --git a/Engine/Source/Editor/UnrealEd/Private/PlayLevel.cpp b/Engine/Source/Editor/UnrealEd/Private/PlayLevel.cpp
--- a/Engine/Source/Editor/UnrealEd/Private/PlayLevel.cpp
+++ b/Engine/Source/Editor/UnrealEd/Private/PlayLevel.cpp
@@ -1305,24 +1305,12 @@ void UEditorEngine::PlayUsingLauncher()
 		LauncherProfile->AddCookedPlatform(PlayUsingLauncherDeviceId.Left(PlayUsingLauncherDeviceId.Find(TEXT("@"))));
 		LauncherProfile->SetForceClose(true);
 		LauncherProfile->SetTimeout(60);
-
-		// Run the same version of the cooker that the editor is compiled as
-#if UE_BUILD_DEBUG
-		LauncherProfile->SetCookConfiguration(EBuildConfigurations::Debug);
-#elif UE_BUILD_DEBUGGAME
-		LauncherProfile->SetCookConfiguration(EBuildConfigurations::DebugGame);
-#elif UE_BUILD_SHIPPING
-		LauncherProfile->SetCookConfiguration(EBuildConfigurations::Shipping);
-#elif UE_BUILD_TEST
-		LauncherProfile->SetCookConfiguration(EBuildConfigurations::Test);
-#else	// UE_BUILD_DEVELOPMENT
-		LauncherProfile->SetCookConfiguration(EBuildConfigurations::Development);
-#endif
-
+		LauncherProfile->SetCookConfiguration(FApp::GetBuildConfiguration());
 		LauncherProfile->SetBuildConfiguration(EBuildConfigurations::Development);
 		LauncherProfile->SetDeploymentMode(ELauncherProfileDeploymentModes::CopyToDevice);
 		LauncherProfile->SetDeployedDeviceGroup(DeviceGroup);
 		LauncherProfile->SetHideFileServerWindow(false);
+		LauncherProfile->SetEditorExe(FUnrealEdMisc::Get().GetExecutableForCommandlets());
 		const FString DummyDeviceName(FString::Printf(TEXT("All_iOS_On_%s"), FPlatformProcess::ComputerName()));
 		if (PlayUsingLauncherDeviceId.Left(PlayUsingLauncherDeviceId.Find(TEXT("@"))) != TEXT("IOS") || !PlayUsingLauncherDeviceName.Contains(DummyDeviceName))
 		{
diff --git a/Engine/Source/Editor/UnrealEd/Private/UnrealEdMisc.cpp b/Engine/Source/Editor/UnrealEd/Private/UnrealEdMisc.cpp
--- a/Engine/Source/Editor/UnrealEd/Private/UnrealEdMisc.cpp
+++ b/Engine/Source/Editor/UnrealEd/Private/UnrealEdMisc.cpp
@@ -1324,6 +1324,23 @@ bool FUnrealEdMisc::GetURL( const TCHAR* InKey, FString& OutURL, const bool bChe
 	return bFound;
 }
 
+FString FUnrealEdMisc::GetExecutableForCommandlets() const
+{
+	FString ExecutableName = FPlatformProcess::ExecutableName(false);
+#if PLATFORM_WINDOWS
+	// turn UE4editor into UE4editor-cmd
+	if(ExecutableName.EndsWith(".exe", ESearchCase::IgnoreCase) && !FPaths::GetBaseFilename(ExecutableName).EndsWith("-cmd", ESearchCase::IgnoreCase))
+	{
+		FString NewExeName = ExecutableName.Left(ExecutableName.Len() - 4) + "-Cmd.exe";
+		if (FPaths::FileExists(NewExeName))
+		{
+			ExecutableName = NewExeName;
+		}
+	}
+#endif
+	return ExecutableName;
+}
+
 void FUnrealEdMisc::OnUserDefinedGestureChanged(const FUICommandInfo& CommandInfo)
 {
 	if( FEngineAnalytics::IsAvailable() )
diff --git a/Engine/Source/Editor/UnrealEd/Public/UnrealEdMisc.h b/Engine/Source/Editor/UnrealEd/Public/UnrealEdMisc.h
--- a/Engine/Source/Editor/UnrealEd/Public/UnrealEdMisc.h
+++ b/Engine/Source/Editor/UnrealEd/Public/UnrealEdMisc.h
@@ -198,6 +198,9 @@ public:
 	 */
 	bool GetURL( const TCHAR* InKey, FString& OutURL, const bool bCheckRocket = false ) const;
 
+	/** Returns the editor executable to use to execute commandlets */
+	FString GetExecutableForCommandlets() const;
+
 private:
 	/** Constructor, private - use Get() function */
 	FUnrealEdMisc();
