fix for UFE using incremental cooking by default
TTP342370
#ue4
#ufe

diff --git a/Engine/Source/Developer/LauncherServices/Private/Profiles/LauncherProfile.h b/Engine/Source/Developer/LauncherServices/Private/Profiles/LauncherProfile.h
--- a/Engine/Source/Developer/LauncherServices/Private/Profiles/LauncherProfile.h
+++ b/Engine/Source/Developer/LauncherServices/Private/Profiles/LauncherProfile.h
@@ -446,7 +446,7 @@ public:
 		CookConfiguration = FApp::GetBuildConfiguration();
 		CookMode = ELauncherProfileCookModes::DoNotCook;
 		CookOptions = FString();
-		CookIncremental = true;
+		CookIncremental = false;
 		CookUnversioned = false;
 		CookedCultures.Reset();
 		CookedCultures.Add(I18N.GetCurrentCulture()->GetName());
