Fix failure when packaging for WindowsNoEditor in UnrealFrontend.
TTP#342375 UFE: COOKING: Failure when packaging for WindowsNoEditor.

diff --git a/Engine/Source/Developer/LauncherServices/Private/Launcher/LauncherPackageCommands.h b/Engine/Source/Developer/LauncherServices/Private/Launcher/LauncherPackageCommands.h
--- a/Engine/Source/Developer/LauncherServices/Private/Launcher/LauncherPackageCommands.h
+++ b/Engine/Source/Developer/LauncherServices/Private/Launcher/LauncherPackageCommands.h
@@ -219,6 +219,12 @@ public:
 		CommandLine += FString::Printf(TEXT(" -cmdline=\"%s -Messaging\""),
 			*InitialMap);
 
+		auto PackageDirectory = ChainState.Profile->GetPackageDirectory();
+		if (PackageDirectory.Len() > 0)
+		{
+			CommandLine += FString::Printf(TEXT(" -stagingdirectory=\"%s\""), *PackageDirectory);
+		}
+
 		CommandLine += FString::Printf(TEXT(" -addcmdline=\"%s -InstanceId=%s -SessionId=%s -SessionOwner=%s -SessionName='%s'%s%s%s\""),
 			*InitialMap,
 			*InstanceId.ToString(),
