Target Platform API added to provide more information about variants and variant support.
Device Services Refactoring to create the concept of a physical device.
Game launcher and device viewer UI changes to show these new concepts.

diff --git a/Engine/Source/Developer/Android/AndroidTargetPlatform/Private/AndroidTargetPlatform.h b/Engine/Source/Developer/Android/AndroidTargetPlatform/Private/AndroidTargetPlatform.h
--- a/Engine/Source/Developer/Android/AndroidTargetPlatform/Private/AndroidTargetPlatform.h
+++ b/Engine/Source/Developer/Android/AndroidTargetPlatform/Private/AndroidTargetPlatform.h
@@ -133,6 +133,10 @@ public:
 	virtual FName GetWaveFormat( class USoundWave* Wave ) const override;
 #endif //WITH_ENGINE
 
+	virtual bool SupportsVariants() const override;
+
+	virtual FText GetVariantTitle() const override;
+
 	DECLARE_DERIVED_EVENT(FAndroidTargetPlatform, ITargetPlatform::FOnTargetDeviceDiscovered, FOnTargetDeviceDiscovered);
 	virtual FOnTargetDeviceDiscovered& OnDeviceDiscovered( ) override
 	{
diff --git a/Engine/Source/Developer/Android/AndroidTargetPlatform/Private/AndroidTargetPlatform.inl b/Engine/Source/Developer/Android/AndroidTargetPlatform/Private/AndroidTargetPlatform.inl
--- a/Engine/Source/Developer/Android/AndroidTargetPlatform/Private/AndroidTargetPlatform.inl
+++ b/Engine/Source/Developer/Android/AndroidTargetPlatform/Private/AndroidTargetPlatform.inl
@@ -8,6 +8,8 @@
 /* FAndroidTargetPlatform structors
  *****************************************************************************/
 
+#define LOCTEXT_NAMESPACE "FAndroidTargetPlatform"
+
 template<class TPlatformProperties>
 inline FAndroidTargetPlatform<TPlatformProperties>::FAndroidTargetPlatform( ) :
 	DeviceDetection(nullptr)
@@ -313,6 +315,17 @@ FName FAndroidTargetPlatform<TPlatformProperties>::GetWaveFormat( class USoundWa
 
 #endif //WITH_ENGINE
 
+template<class TPlatformProperties>
+bool FAndroidTargetPlatform<TPlatformProperties>::SupportsVariants() const
+{
+	return true;
+}
+
+template<class TPlatformProperties>
+FText FAndroidTargetPlatform<TPlatformProperties>::GetVariantTitle() const
+{
+	return LOCTEXT("AndroidVariantTitle", "Texture Format");
+}
 
 /* FAndroidTargetPlatform implementation
  *****************************************************************************/
@@ -392,3 +405,5 @@ inline bool FAndroidTargetPlatform<TPlatformProperties>::HandleTicker( float Del
 
 	return true;
 }
+
+#undef LOCTEXT_NAMESPACE
diff --git a/Engine/Source/Developer/Android/Android_ATCTargetPlatform/Private/Android_ATCTargetPlatformModule.cpp b/Engine/Source/Developer/Android/Android_ATCTargetPlatform/Private/Android_ATCTargetPlatformModule.cpp
--- a/Engine/Source/Developer/Android/Android_ATCTargetPlatform/Private/Android_ATCTargetPlatformModule.cpp
+++ b/Engine/Source/Developer/Android/Android_ATCTargetPlatform/Private/Android_ATCTargetPlatformModule.cpp
@@ -45,6 +45,16 @@ class FAndroid_ATCTargetPlatform
 	{
 		return (ExtensionsString.Contains(TEXT("GL_ATI_texture_compression_atitc")) || ExtensionsString.Contains(TEXT("GL_AMD_compressed_ATC_texture")));
 	}
+
+	virtual FText GetVariantDisplayName() const override
+	{
+		return LOCTEXT("Android_ATC_ShortName", "ATC");
+	}
+
+	virtual float GetVariantPriority() const override
+	{
+		return 0.5f;
+	}
 };
 
 
diff --git a/Engine/Source/Developer/Android/Android_DXTTargetPlatform/Private/Android_DXTTargetPlatformModule.cpp b/Engine/Source/Developer/Android/Android_DXTTargetPlatform/Private/Android_DXTTargetPlatformModule.cpp
--- a/Engine/Source/Developer/Android/Android_DXTTargetPlatform/Private/Android_DXTTargetPlatformModule.cpp
+++ b/Engine/Source/Developer/Android/Android_DXTTargetPlatform/Private/Android_DXTTargetPlatformModule.cpp
@@ -45,6 +45,16 @@ class FAndroid_DXTTargetPlatform
 	{
 		return (ExtensionsString.Contains(TEXT("GL_NV_texture_compression_s3tc")) || ExtensionsString.Contains(TEXT("GL_EXT_texture_compression_s3tc")));
 	}
+
+	virtual FText GetVariantDisplayName() const override
+	{
+		return LOCTEXT("Android_DXT_ShortName", "DXT");
+	}
+
+	virtual float GetVariantPriority() const override
+	{
+		return 0.6f;
+	}
 };
 
 /**
diff --git a/Engine/Source/Developer/Android/Android_ETC1TargetPlatform/Private/Android_ETC1TargetPlatformModule.cpp b/Engine/Source/Developer/Android/Android_ETC1TargetPlatform/Private/Android_ETC1TargetPlatformModule.cpp
--- a/Engine/Source/Developer/Android/Android_ETC1TargetPlatform/Private/Android_ETC1TargetPlatformModule.cpp
+++ b/Engine/Source/Developer/Android/Android_ETC1TargetPlatform/Private/Android_ETC1TargetPlatformModule.cpp
@@ -51,6 +51,16 @@ public:
 	{
 		return ExtensionsString.Contains(TEXT("GL_OES_compressed_ETC1_RGB8_texture"));
 	}
+
+	virtual FText GetVariantDisplayName() const override
+	{
+		return LOCTEXT("Android_ETC1_ShortName", "ETC1");
+	}
+
+	virtual float GetVariantPriority() const override
+	{
+		return 0.1f;
+	}
 };
 
 
diff --git a/Engine/Source/Developer/Android/Android_ETC2TargetPlatform/Private/Android_ETC2TargetPlatformModule.cpp b/Engine/Source/Developer/Android/Android_ETC2TargetPlatform/Private/Android_ETC2TargetPlatformModule.cpp
--- a/Engine/Source/Developer/Android/Android_ETC2TargetPlatform/Private/Android_ETC2TargetPlatformModule.cpp
+++ b/Engine/Source/Developer/Android/Android_ETC2TargetPlatform/Private/Android_ETC2TargetPlatformModule.cpp
@@ -52,6 +52,16 @@ public:
 	{
 		return GLESVersion >= 0x30000;
 	}
+
+	virtual FText GetVariantDisplayName() const override
+	{
+		return LOCTEXT("Android_ETC2_ShortName", "ETC2");
+	}
+
+	virtual float GetVariantPriority() const override
+	{
+		return 0.2f;
+	}
 };
 
 
diff --git a/Engine/Source/Developer/Android/Android_PVRTCTargetPlatform/Private/Android_PVRTCTargetPlatformModule.cpp b/Engine/Source/Developer/Android/Android_PVRTCTargetPlatform/Private/Android_PVRTCTargetPlatformModule.cpp
--- a/Engine/Source/Developer/Android/Android_PVRTCTargetPlatform/Private/Android_PVRTCTargetPlatformModule.cpp
+++ b/Engine/Source/Developer/Android/Android_PVRTCTargetPlatform/Private/Android_PVRTCTargetPlatformModule.cpp
@@ -44,6 +44,16 @@ class FAndroid_PVRTCTargetPlatform
 	{
 		return ExtensionsString.Contains(TEXT("GL_IMG_texture_compression_pvrtc"));
 	}
+
+	virtual FText GetVariantDisplayName() const override
+	{
+		return LOCTEXT("Android_PVRTC_ShortName", "PVRTC");
+	}
+
+	virtual float GetVariantPriority() const override
+	{
+		return 0.9f;
+	}
 };
 
 /**
diff --git a/Engine/Source/Developer/DeviceManager/Private/Models/DeviceBrowserFilter.h b/Engine/Source/Developer/DeviceManager/Private/Models/DeviceBrowserFilter.h
--- a/Engine/Source/Developer/DeviceManager/Private/Models/DeviceBrowserFilter.h
+++ b/Engine/Source/Developer/DeviceManager/Private/Models/DeviceBrowserFilter.h
@@ -9,6 +9,17 @@ typedef TSharedPtr<class FDeviceBrowserFilter> FDeviceBrowserFilterPtr;
 /** Type definition for shared references to instances of FDeviceBrowserFilter. */
 typedef TSharedRef<class FDeviceBrowserFilter> FDeviceBrowserFilterRef;
 
+class FDeviceBrowserFilterEntry
+{
+public:
+	FDeviceBrowserFilterEntry(FString InPlatformName, FName InPlatformLookup)
+		: PlatformName(InPlatformName)
+		, PlatformLookup(InPlatformLookup)
+	{
+	}
+	FString PlatformName;
+	FName PlatformLookup;
+};
 
 /**
  * Implements a filter for the device browser's target device service list.
@@ -27,7 +38,7 @@ public:
 	{
 		if (DeviceService.IsValid())
 		{
-			if (DisabledPlatforms.Contains(DeviceService->GetDeviceId().GetPlatformName()))
+			if (DisabledPlatforms.Contains(DeviceService->GetDevicePlatformDisplayName()))
 			{
 				return false;
 			}
@@ -37,7 +48,7 @@ public:
 				return true;
 			}
 
-			return (DeviceService->GetCachedDeviceName().Contains(DeviceSearchText.ToString()) );
+			return (DeviceService->GetDeviceName().Contains(DeviceSearchText.ToString()) );
 		}
 
 		return false;
@@ -69,7 +80,7 @@ public:
 	 *
 	 * @return List of platform names.
 	 */
-	const TArray<TSharedPtr<FString> >& GetFilteredPlatforms( ) const
+	const TArray<TSharedPtr<FDeviceBrowserFilterEntry>> & GetFilteredPlatforms() const
 	{
 		return PlatformList;
 	}
@@ -100,12 +111,12 @@ public:
 			const ITargetDeviceServicePtr& DeviceService = DeviceServices[DeviceServiceIndex];
 
 			// populate platforms
-			const FString& Platform = DeviceService->GetDeviceId().GetPlatformName();
+			const FString Platform = DeviceService->GetDevicePlatformDisplayName();
 			int32& PlatformCounter = PlatformCounters.FindOrAdd(Platform);
 
 			if (PlatformCounter == 0)
 			{
-				PlatformList.Add(MakeShareable(new FString(Platform)));
+				PlatformList.Add(MakeShareable(new FDeviceBrowserFilterEntry(Platform, DeviceService->GetDevicePlatformName())));
 			}
 
 			++PlatformCounter;
@@ -188,7 +199,7 @@ private:
 	TMap<FString, int32> PlatformCounters;
 
 	// Holds the list of platform filters.
-	TArray<TSharedPtr<FString>> PlatformList;
+	TArray<TSharedPtr<FDeviceBrowserFilterEntry>> PlatformList;
 
 private:
 
diff --git a/Engine/Source/Developer/DeviceManager/Private/Widgets/Browser/SDeviceBrowser.cpp b/Engine/Source/Developer/DeviceManager/Private/Widgets/Browser/SDeviceBrowser.cpp
--- a/Engine/Source/Developer/DeviceManager/Private/Widgets/Browser/SDeviceBrowser.cpp
+++ b/Engine/Source/Developer/DeviceManager/Private/Widgets/Browser/SDeviceBrowser.cpp
@@ -25,15 +25,16 @@ void SDeviceBrowser::Construct( const FArguments& InArgs, const FDeviceManagerMo
 		+ SVerticalBox::Slot()
 			.AutoHeight()
 				[
+					/*
 					SNew(SExpandableArea)
 						.AreaTitle(LOCTEXT("FilterBarAreaTitle", "Device Filter").ToString())
 						.InitiallyCollapsed(true)
 						.Padding(FMargin(8.0f, 8.0f, 8.0f, 4.0f))
 						.BodyContent()
-						[
+						[*/
 							// filter bar
 							SNew(SDeviceBrowserFilterBar, Filter.ToSharedRef())
-						]
+						//]
 				]
 
 		+ SVerticalBox::Slot()
diff --git a/Engine/Source/Developer/DeviceManager/Private/Widgets/Browser/SDeviceBrowserDeviceAdder.cpp b/Engine/Source/Developer/DeviceManager/Private/Widgets/Browser/SDeviceBrowserDeviceAdder.cpp
--- a/Engine/Source/Developer/DeviceManager/Private/Widgets/Browser/SDeviceBrowserDeviceAdder.cpp
+++ b/Engine/Source/Developer/DeviceManager/Private/Widgets/Browser/SDeviceBrowserDeviceAdder.cpp
@@ -212,9 +212,10 @@ void SDeviceBrowserDeviceAdder::RefreshPlatformList( )
 
 FReply SDeviceBrowserDeviceAdder::HandleAddButtonClicked( )
 {
-	FTargetDeviceId DeviceId(*PlatformComboBox->GetSelectedItem(), DeviceIdTextBox->GetText().ToString());
+	ITargetPlatform * TargetPlatform = GetTargetPlatformManager()->FindTargetPlatform(*PlatformComboBox->GetSelectedItem());
 
-	if (DeviceServiceManager->AddStartupService(DeviceId, DeviceNameTextBox->GetText().ToString()))
+	bool bAdded = TargetPlatform->AddDevice(DeviceIdTextBox->GetText().ToString(), false);
+	if (bAdded)
 	{
 		DeviceIdTextBox->SetText(FText::GetEmpty());
 		DeviceNameTextBox->SetText(FText::GetEmpty());
diff --git a/Engine/Source/Developer/DeviceManager/Private/Widgets/Browser/SDeviceBrowserDeviceListRow.h b/Engine/Source/Developer/DeviceManager/Private/Widgets/Browser/SDeviceBrowserDeviceListRow.h
--- a/Engine/Source/Developer/DeviceManager/Private/Widgets/Browser/SDeviceBrowserDeviceListRow.h
+++ b/Engine/Source/Developer/DeviceManager/Private/Widgets/Browser/SDeviceBrowserDeviceListRow.h
@@ -61,7 +61,7 @@ public:
 		}
 		else if (ColumnName == TEXT("Icon"))
 		{
-			const PlatformInfo::FPlatformInfo* const PlatformInfo = PlatformInfo::FindPlatformInfo(*DeviceService->GetDeviceId().GetPlatformName());
+			const PlatformInfo::FPlatformInfo* const PlatformInfo = PlatformInfo::FindPlatformInfo(DeviceService->GetDevicePlatformName());
 
 			return SNew(SBox)
 				.Padding(FMargin(4.0f, 0.0f))
@@ -92,7 +92,7 @@ public:
 				[
 					SNew(STextBlock)
 						.ColorAndOpacity(this, &SDeviceBrowserDeviceListRow::HandleTextColorAndOpacity)
-						.Text(DeviceService->GetDeviceId().GetPlatformName())
+						.Text(DeviceService->GetDevicePlatformDisplayName())
 				];
 		}
 		else if (ColumnName == TEXT("Share"))
@@ -142,7 +142,7 @@ private:
 	// Callback for getting the text in the 'Name' column.
 	FText HandleNameText( ) const
 	{
-		return FText::FromString(DeviceService->GetCachedDeviceName());
+		return FText::FromString(DeviceService->GetDeviceName());
 	}
 
 	// Callback for changing this row's Share check box state.
diff --git a/Engine/Source/Developer/DeviceManager/Private/Widgets/Browser/SDeviceBrowserFilterBar.cpp b/Engine/Source/Developer/DeviceManager/Private/Widgets/Browser/SDeviceBrowserFilterBar.cpp
--- a/Engine/Source/Developer/DeviceManager/Private/Widgets/Browser/SDeviceBrowserFilterBar.cpp
+++ b/Engine/Source/Developer/DeviceManager/Private/Widgets/Browser/SDeviceBrowserFilterBar.cpp
@@ -30,37 +30,41 @@ void SDeviceBrowserFilterBar::Construct( const FArguments& InArgs, FDeviceBrowse
 	[
 		SNew(SHorizontalBox)
 
-		+ SHorizontalBox::Slot()
-			.FillWidth(1.0f)
-			.VAlign(VAlign_Top)
-			[
-				// search box
-				SAssignNew(FilterStringTextBox, SSearchBox)
-					.HintText(LOCTEXT("SearchBoxHint", "Search devices"))
-					.OnTextChanged(this, &SDeviceBrowserFilterBar::HandleFilterStringTextChanged)
-			]
-
 		+ SHorizontalBox::Slot()
 			.AutoWidth()
 			.HAlign(HAlign_Right)
-			.Padding(6.0f, 0.0f, 0.0f, 0.0f)
+			.Padding(6.0f, 0.0f, 6.0f, 0.0f)
 			[
 				// platform filter
 				SNew(SComboButton)
+					.ComboButtonStyle(FEditorStyle::Get(), "ContentBrowser.Filters.Style")
+					.ForegroundColor(FLinearColor::White)
 					.ButtonContent()
 					[
 						SNew(STextBlock)
-							.Text(LOCTEXT("PlatformComboButtonText", "Platforms"))
+						.TextStyle(FEditorStyle::Get(), "ContentBrowser.Filters.Text")
+						.Text(LOCTEXT("PlatformFiltersComboButtonText", "Platform Filters"))
 					]
-					.ContentPadding(FMargin(6.0f, 2.0f))
+					.ContentPadding(0.0f)
 					.MenuContent()
 					[
-						SAssignNew(PlatformListView, SListView<TSharedPtr<FString> >)
+						SAssignNew(PlatformListView, SListView<TSharedPtr<FDeviceBrowserFilterEntry> >)
 							.ItemHeight(24.0f)
 							.ListItemsSource(&Filter->GetFilteredPlatforms())
 							.OnGenerateRow(this, &SDeviceBrowserFilterBar::HandlePlatformListViewGenerateRow)
 					]
 			]
+
+		+ SHorizontalBox::Slot()
+			.FillWidth(1.0f)
+			.VAlign(VAlign_Top)
+			[
+				// search box
+				SAssignNew(FilterStringTextBox, SSearchBox)
+				.HintText(LOCTEXT("SearchBoxHint", "Search devices"))
+				.OnTextChanged(this, &SDeviceBrowserFilterBar::HandleFilterStringTextChanged)
+			]
+
 	];
 
 	Filter->OnFilterReset().AddSP(this, &SDeviceBrowserFilterBar::HandleFilterReset);
@@ -84,15 +88,15 @@ void SDeviceBrowserFilterBar::HandleFilterStringTextChanged( const FText& NewTex
 }
 
 
-void SDeviceBrowserFilterBar::HandlePlatformListRowCheckStateChanged( ESlateCheckBoxState::Type CheckState, TSharedPtr<FString> PlatformName )
+void SDeviceBrowserFilterBar::HandlePlatformListRowCheckStateChanged(ESlateCheckBoxState::Type CheckState, TSharedPtr<FDeviceBrowserFilterEntry> PlatformEntry)
 {
-	Filter->SetPlatformEnabled(*PlatformName, CheckState == ESlateCheckBoxState::Checked);
+	Filter->SetPlatformEnabled(PlatformEntry->PlatformName, CheckState == ESlateCheckBoxState::Checked);
 }
 
 
-ESlateCheckBoxState::Type SDeviceBrowserFilterBar::HandlePlatformListRowIsChecked( TSharedPtr<FString> PlatformName ) const
+ESlateCheckBoxState::Type SDeviceBrowserFilterBar::HandlePlatformListRowIsChecked(TSharedPtr<FDeviceBrowserFilterEntry> PlatformEntry) const
 {
-	if (Filter->IsPlatformEnabled(*PlatformName))
+	if (Filter->IsPlatformEnabled(PlatformEntry->PlatformName))
 	{
 		return ESlateCheckBoxState::Checked;
 	}
@@ -101,48 +105,48 @@ ESlateCheckBoxState::Type SDeviceBrowserFilterBar::HandlePlatformListRowIsChecke
 }
 
 
-FString SDeviceBrowserFilterBar::HandlePlatformListRowText( TSharedPtr<FString> PlatformName ) const
+FString SDeviceBrowserFilterBar::HandlePlatformListRowText(TSharedPtr<FDeviceBrowserFilterEntry> PlatformEntry) const
 {
-	return FString::Printf(TEXT("%s (%i)"), **PlatformName, Filter->GetServiceCountPerPlatform(*PlatformName));
+	return FString::Printf(TEXT("%s (%i)"), *(PlatformEntry->PlatformName), Filter->GetServiceCountPerPlatform(PlatformEntry->PlatformName));
 }
 
 BEGIN_SLATE_FUNCTION_BUILD_OPTIMIZATION
-TSharedRef<ITableRow> SDeviceBrowserFilterBar::HandlePlatformListViewGenerateRow( TSharedPtr<FString> PlatformName, const TSharedRef<STableViewBase>& OwnerTable )
+TSharedRef<ITableRow> SDeviceBrowserFilterBar::HandlePlatformListViewGenerateRow(TSharedPtr<FDeviceBrowserFilterEntry> PlatformEntry, const TSharedRef<STableViewBase>& OwnerTable)
 {
-	const PlatformInfo::FPlatformInfo* const PlatformInfo = PlatformInfo::FindPlatformInfo(**PlatformName);
+	const PlatformInfo::FPlatformInfo* const PlatformInfo = PlatformInfo::FindPlatformInfo(PlatformEntry->PlatformLookup);
 
 	return SNew(STableRow<TSharedPtr<FString> >, OwnerTable)
 		.Content()
 		[
 			SNew(SCheckBox)
-				.IsChecked(this, &SDeviceBrowserFilterBar::HandlePlatformListRowIsChecked, PlatformName)
-				.Padding(FMargin(6.0, 2.0))
-				.OnCheckStateChanged(this, &SDeviceBrowserFilterBar::HandlePlatformListRowCheckStateChanged, PlatformName)
-				.Content()
+			.IsChecked(this, &SDeviceBrowserFilterBar::HandlePlatformListRowIsChecked, PlatformEntry)
+			.Padding(FMargin(6.0, 2.0))
+			.OnCheckStateChanged(this, &SDeviceBrowserFilterBar::HandlePlatformListRowCheckStateChanged, PlatformEntry)
+			.Content()
+			[
+				SNew(SHorizontalBox)
+
+				+ SHorizontalBox::Slot()
+				.AutoWidth()
 				[
-					SNew(SHorizontalBox)
-
-					+ SHorizontalBox::Slot()
-						.AutoWidth()
-						[
-							SNew(SBox)
-								.WidthOverride(24)
-								.HeightOverride(24)
-								[
-									SNew(SImage)
-										.Image((PlatformInfo) ? FEditorStyle::GetBrush(PlatformInfo->GetIconStyleName(PlatformInfo::EPlatformIconSize::Normal)) : FStyleDefaults::GetNoBrush())
-								]
-						]
-
-					+ SHorizontalBox::Slot()
-						.FillWidth(1.0f)
-						.Padding(4.0f, 0.0f, 0.0f, 0.0f)
-						.VAlign(VAlign_Center)
-						[
-							SNew(STextBlock)
-								.Text(this, &SDeviceBrowserFilterBar::HandlePlatformListRowText, PlatformName)
-						]
+					SNew(SBox)
+					.WidthOverride(24)
+					.HeightOverride(24)
+					[
+						SNew(SImage)
+						.Image((PlatformInfo) ? FEditorStyle::GetBrush(PlatformInfo->GetIconStyleName(PlatformInfo::EPlatformIconSize::Normal)) : FStyleDefaults::GetNoBrush())
+					]
+				]
+
+				+ SHorizontalBox::Slot()
+				.FillWidth(1.0f)
+				.Padding(4.0f, 0.0f, 0.0f, 0.0f)
+				.VAlign(VAlign_Center)
+				[
+					SNew(STextBlock)
+					.Text(this, &SDeviceBrowserFilterBar::HandlePlatformListRowText, PlatformEntry)
 				]
+			]
 		];
 }
 END_SLATE_FUNCTION_BUILD_OPTIMIZATION
diff --git a/Engine/Source/Developer/DeviceManager/Private/Widgets/Browser/SDeviceBrowserFilterBar.h b/Engine/Source/Developer/DeviceManager/Private/Widgets/Browser/SDeviceBrowserFilterBar.h
--- a/Engine/Source/Developer/DeviceManager/Private/Widgets/Browser/SDeviceBrowserFilterBar.h
+++ b/Engine/Source/Developer/DeviceManager/Private/Widgets/Browser/SDeviceBrowserFilterBar.h
@@ -19,14 +19,14 @@ public:
 	/**
 	 * Default constructor.
 	 */
-	SDeviceBrowserFilterBar( )
+	SDeviceBrowserFilterBar()
 		: Filter(NULL)
 	{ }
 
 	/**
 	 * Destructor.
 	 */
-	~SDeviceBrowserFilterBar( );
+	~SDeviceBrowserFilterBar();
 
 public:
 
@@ -36,27 +36,27 @@ public:
 	 * @param InArgs The declaration data for this widget.
 	 * @param InFilter The filter model to use.
 	 */
-	void Construct( const FArguments& InArgs, FDeviceBrowserFilterRef InFilter );
+	void Construct(const FArguments& InArgs, FDeviceBrowserFilterRef InFilter);
 
 private:
 
 	// Callback for filter model resets.
-	void HandleFilterReset( );
+	void HandleFilterReset();
 
 	// Callback for changing the filter string text box text.
-	void HandleFilterStringTextChanged( const FText& NewText );
+	void HandleFilterStringTextChanged(const FText& NewText);
 
 	// Callback for changing the checked state of the given platform filter row.
-	void HandlePlatformListRowCheckStateChanged( ESlateCheckBoxState::Type CheckState, TSharedPtr<FString> PlatformName );
+	void HandlePlatformListRowCheckStateChanged(ESlateCheckBoxState::Type CheckState, TSharedPtr<FDeviceBrowserFilterEntry> PlatformEntry);
 
 	// Callback for getting the checked state of the given platform filter row.
-	ESlateCheckBoxState::Type HandlePlatformListRowIsChecked( TSharedPtr<FString> PlatformName ) const;
+	ESlateCheckBoxState::Type HandlePlatformListRowIsChecked(TSharedPtr<FDeviceBrowserFilterEntry> PlatformEntry) const;
 
 	// Callback for getting the text for a row in the platform filter drop-down.
-	FString HandlePlatformListRowText( TSharedPtr<FString> PlatformName ) const;
+	FString HandlePlatformListRowText(TSharedPtr<FDeviceBrowserFilterEntry> PlatformEntry) const;
 
 	// Generates a row widget for the platform filter list.
-	TSharedRef<ITableRow> HandlePlatformListViewGenerateRow( TSharedPtr<FString> PlatformName, const TSharedRef<STableViewBase>& OwnerTable );
+	TSharedRef<ITableRow> HandlePlatformListViewGenerateRow(TSharedPtr<FDeviceBrowserFilterEntry> PlatformEntry, const TSharedRef<STableViewBase>& OwnerTable);
 
 private:
 
@@ -67,5 +67,5 @@ private:
 	TSharedPtr<SSearchBox> FilterStringTextBox;
 
 	// Holds the platform filters list view.
-	TSharedPtr<SListView<TSharedPtr<FString>>> PlatformListView;
+	TSharedPtr<SListView<TSharedPtr<FDeviceBrowserFilterEntry>>> PlatformListView;
 };
diff --git a/Engine/Source/Developer/DeviceManager/Private/Widgets/SDeviceManager.cpp b/Engine/Source/Developer/DeviceManager/Private/Widgets/SDeviceManager.cpp
--- a/Engine/Source/Developer/DeviceManager/Private/Widgets/SDeviceManager.cpp
+++ b/Engine/Source/Developer/DeviceManager/Private/Widgets/SDeviceManager.cpp
@@ -465,7 +465,7 @@ void SDeviceManager::HandleRemoveActionExecute( )
 
 	if (DeviceService.IsValid())
 	{
-		DeviceServiceManager->RemoveStartupService(DeviceService->GetDeviceId());
+		DeviceServiceManager->RemoveStartupService(DeviceService->GetDeviceName());
 	}
 }
 
diff --git a/Engine/Source/Developer/DeviceManager/Private/Widgets/Shared/SDeviceQuickInfo.h b/Engine/Source/Developer/DeviceManager/Private/Widgets/Shared/SDeviceQuickInfo.h
--- a/Engine/Source/Developer/DeviceManager/Private/Widgets/Shared/SDeviceQuickInfo.h
+++ b/Engine/Source/Developer/DeviceManager/Private/Widgets/Shared/SDeviceQuickInfo.h
@@ -175,7 +175,11 @@ private:
 	{
 		if (DeviceService.IsValid())
 		{
-			return FText::FromString(DeviceService->GetDeviceId().ToString());
+			ITargetDevicePtr Device = DeviceService->GetDevice();
+			if (Device.IsValid())
+			{
+				return FText::FromString(Device->GetId().ToString());
+			}
 		}
 
 		return LOCTEXT("UnknownValue", "<unknown>");
@@ -186,7 +190,7 @@ private:
 	{
 		if (DeviceService.IsValid())
 		{
-			FString DeviceName = DeviceService->GetCachedDeviceName();
+			FString DeviceName = DeviceService->GetDeviceName();
 
 			if (!DeviceName.IsEmpty())
 			{
@@ -241,7 +245,7 @@ private:
 	{
 		if (DeviceService.IsValid())
 		{
-			const PlatformInfo::FPlatformInfo* const PlatformInfo = PlatformInfo::FindPlatformInfo(*DeviceService->GetDeviceId().GetPlatformName());
+			const PlatformInfo::FPlatformInfo* const PlatformInfo = PlatformInfo::FindPlatformInfo(DeviceService->GetDevicePlatformName());
 			if(PlatformInfo)
 			{
 				return FEditorStyle::GetBrush(PlatformInfo->GetIconStyleName(PlatformInfo::EPlatformIconSize::XLarge));
@@ -256,7 +260,7 @@ private:
 	{
 		if (DeviceService.IsValid())
 		{
-			FString PlatformName = DeviceService->GetDeviceId().GetPlatformName();
+			FString PlatformName = DeviceService->GetDevicePlatformDisplayName();
 
 			if (!PlatformName.IsEmpty())
 			{
diff --git a/Engine/Source/Developer/LauncherServices/Private/Launcher/LauncherCookCommands.h b/Engine/Source/Developer/LauncherServices/Private/Launcher/LauncherCookCommands.h
--- a/Engine/Source/Developer/LauncherServices/Private/Launcher/LauncherCookCommands.h
+++ b/Engine/Source/Developer/LauncherServices/Private/Launcher/LauncherCookCommands.h
@@ -55,6 +55,13 @@ public:
 		{
 			CommandLine += TEXT(" -Unversioned");
 		}
+
+		FString additionalOptions = ChainState.Profile->GetCookOptions();
+		if (!additionalOptions.IsEmpty())
+		{
+			CommandLine += TEXT(" ");
+			CommandLine += additionalOptions;
+		}
 		
 		return CommandLine;
 	}
diff --git a/Engine/Source/Developer/LauncherServices/Private/Launcher/LauncherDeployCommands.h b/Engine/Source/Developer/LauncherServices/Private/Launcher/LauncherDeployCommands.h
--- a/Engine/Source/Developer/LauncherServices/Private/Launcher/LauncherDeployCommands.h
+++ b/Engine/Source/Developer/LauncherServices/Private/Launcher/LauncherDeployCommands.h
@@ -10,8 +10,9 @@ class FLauncherDeployGameToDeviceCommand
 	: public FLauncherUATCommand
 {
 public:
-	FLauncherDeployGameToDeviceCommand( const ITargetDeviceProxyRef& InDeviceProxy, const ITargetPlatform& InTargetPlatform, const TSharedPtr<FLauncherUATCommand>& InCook, const FString& InCmdLine )
+	FLauncherDeployGameToDeviceCommand(const ITargetDeviceProxyRef& InDeviceProxy, FName InFlavor, const ITargetPlatform& InTargetPlatform, const TSharedPtr<FLauncherUATCommand>& InCook, const FString& InCmdLine)
 		: DeviceProxy(InDeviceProxy)
+		, Flavor(InFlavor)
 		, TargetPlatform(InTargetPlatform)
 		, InstanceId(FGuid::NewGuid())
 		, CookCommand(InCook)
@@ -60,7 +61,7 @@ public:
 		if (TargetPlatform.PlatformName() != TEXT("IOS"))
 #endif
 		{
-			CommandLine += FString::Printf(TEXT(" -device=\"%s\""), *DeviceProxy->GetDeviceId());
+			CommandLine += FString::Printf(TEXT(" -device=\"%s\""), *DeviceProxy->GetTargetDeviceId(Flavor));
 		}
 
 		// cook dependency arguments
@@ -103,6 +104,9 @@ private:
 	// Holds a pointer to the device proxy to deploy to.
 	ITargetDeviceProxyPtr DeviceProxy;
 
+	// Holds the name of the flavor of Target Device to use.
+	FName Flavor;
+
 	// Holds a pointer to the target platform.
 	const ITargetPlatform& TargetPlatform;
 
@@ -125,8 +129,9 @@ class FLauncherDeployServerToDeviceCommand
 {
 public:
 
-	FLauncherDeployServerToDeviceCommand( const ITargetDeviceProxyRef& InDeviceProxy, const ITargetPlatform& InTargetPlatform, const TSharedPtr<FLauncherUATCommand>& InCook )
+	FLauncherDeployServerToDeviceCommand( const ITargetDeviceProxyRef& InDeviceProxy, FName InFlavor, const ITargetPlatform& InTargetPlatform, const TSharedPtr<FLauncherUATCommand>& InCook )
 		: DeviceProxy(InDeviceProxy)
+		, Flavor(InFlavor)
 		, TargetPlatform(InTargetPlatform)
 		, InstanceId(FGuid::NewGuid())
 		, CookCommand(InCook)
@@ -170,8 +175,8 @@ public:
 			*StagePath,
 			*InitialMap,
 			*TargetPlatform.PlatformName());
-		CommandLine += FString::Printf(TEXT(" -device=\"%s\""), *DeviceProxy->GetDeviceId());
-		CommandLine += FString::Printf(TEXT(" -serverdevice=\"%s\""), *DeviceProxy->GetDeviceId());
+		CommandLine += FString::Printf(TEXT(" -device=\"%s\""), *DeviceProxy->GetTargetDeviceId(Flavor));
+		CommandLine += FString::Printf(TEXT(" -serverdevice=\"%s\""), *DeviceProxy->GetTargetDeviceId(Flavor));
 
 		// cook dependency arguments
 		CommandLine += CookCommand.IsValid() ? CookCommand->GetDependencyArguments(ChainState) : TEXT(" -skipcook");
@@ -203,6 +208,9 @@ private:
 	// Holds a pointer to the device proxy to deploy to.
 	ITargetDeviceProxyPtr DeviceProxy;
 
+	// Holds the name of the flavor of Target Device to use.
+	FName Flavor;
+
 	// Holds a pointer to the target platform.
 	const ITargetPlatform& TargetPlatform;
 
@@ -222,8 +230,9 @@ class FLauncherDeployGamePackageToDeviceCommand
 {
 public:
 
-	FLauncherDeployGamePackageToDeviceCommand( const ITargetDeviceProxyRef& InDeviceProxy, const ITargetPlatform& InTargetPlatform, const TSharedPtr<FLauncherUATCommand>& InCook, const FString& InCmdLine )
+	FLauncherDeployGamePackageToDeviceCommand(const ITargetDeviceProxyRef& InDeviceProxy, FName InFlavor, const ITargetPlatform& InTargetPlatform, const TSharedPtr<FLauncherUATCommand>& InCook, const FString& InCmdLine)
 		: DeviceProxy(InDeviceProxy)
+		, Flavor(InFlavor)
 		, TargetPlatform(InTargetPlatform)
 		, InstanceId(FGuid::NewGuid())
 		, CookCommand(InCook)
@@ -255,7 +264,7 @@ public:
 			*ChainState.Profile->GetPackageDirectory(),
 			*InitialMap,
 			*TargetPlatform.PlatformName());
-		CommandLine += FString::Printf(TEXT(" -device=\"%s\""), *DeviceProxy->GetDeviceId());
+		CommandLine += FString::Printf(TEXT(" -device=\"%s\""), *DeviceProxy->GetTargetDeviceId(Flavor));
 
 		// cook dependency arguments
 		CommandLine += CookCommand.IsValid() ? CookCommand->GetDependencyArguments(ChainState) : TEXT(" -skipcook");
@@ -296,6 +305,9 @@ private:
 	// Holds a pointer to the device proxy to deploy to.
 	ITargetDeviceProxyPtr DeviceProxy;
 
+	// Holds the name of the flavor of Target Device to use.
+	FName Flavor;
+
 	// Holds a pointer to the target platform.
 	const ITargetPlatform& TargetPlatform;
 
@@ -318,8 +330,9 @@ class FLauncherDeployServerPackageToDeviceCommand
 {
 public:
 
-	FLauncherDeployServerPackageToDeviceCommand( const ITargetDeviceProxyRef& InDeviceProxy, const ITargetPlatform& InTargetPlatform, const TSharedPtr<FLauncherUATCommand>& InCook )
+	FLauncherDeployServerPackageToDeviceCommand(const ITargetDeviceProxyRef& InDeviceProxy, FName InFlavor, const ITargetPlatform& InTargetPlatform, const TSharedPtr<FLauncherUATCommand>& InCook)
 		: DeviceProxy(InDeviceProxy)
+		, Flavor(InFlavor)
 		, TargetPlatform(InTargetPlatform)
 		, InstanceId(FGuid::NewGuid())
 		, CookCommand(InCook)
@@ -360,8 +373,8 @@ public:
 			*ChainState.Profile->GetPackageDirectory(),
 			*InitialMap,
 			*TargetPlatform.PlatformName());
-		CommandLine += FString::Printf(TEXT(" -device=\"%s\""), *DeviceProxy->GetDeviceId());
-		CommandLine += FString::Printf(TEXT(" -serverdevice=\"%s\""), *DeviceProxy->GetDeviceId());
+		CommandLine += FString::Printf(TEXT(" -device=\"%s\""), *DeviceProxy->GetTargetDeviceId(Flavor));
+		CommandLine += FString::Printf(TEXT(" -serverdevice=\"%s\""), *DeviceProxy->GetTargetDeviceId(Flavor));
 
 		// cook dependency arguments
 		CommandLine += CookCommand.IsValid() ? CookCommand->GetDependencyArguments(ChainState) : TEXT(" -skipcook");
@@ -393,6 +406,9 @@ private:
 	// Holds a pointer to the device proxy to deploy to.
 	ITargetDeviceProxyPtr DeviceProxy;
 
+	// Holds the name of the flavor of Target Device to use.
+	FName Flavor;
+
 	// Holds a pointer to the target platform.
 	const ITargetPlatform& TargetPlatform;
 
diff --git a/Engine/Source/Developer/LauncherServices/Private/Launcher/LauncherLaunchCommands.h b/Engine/Source/Developer/LauncherServices/Private/Launcher/LauncherLaunchCommands.h
--- a/Engine/Source/Developer/LauncherServices/Private/Launcher/LauncherLaunchCommands.h
+++ b/Engine/Source/Developer/LauncherServices/Private/Launcher/LauncherLaunchCommands.h
@@ -11,8 +11,9 @@ class FLauncherLaunchGameCommand
 {
 public:
 
-	FLauncherLaunchGameCommand(const ITargetDeviceProxyRef& InDeviceProxy, const ITargetPlatform& InTargetPlatform, const ILauncherProfileLaunchRoleRef& InRole, const TSharedPtr<FLauncherUATCommand>& InCook)
+	FLauncherLaunchGameCommand(const ITargetDeviceProxyRef& InDeviceProxy, FName InFlavor, const ITargetPlatform& InTargetPlatform, const ILauncherProfileLaunchRoleRef& InRole, const TSharedPtr<FLauncherUATCommand>& InCook)
 		: DeviceProxy(InDeviceProxy)
+		, Variant(InFlavor)
 		, TargetPlatform(InTargetPlatform)
 		, InstanceId(FGuid::NewGuid())
 		, Role(InRole)
@@ -56,7 +57,7 @@ public:
 		CommandLine += FString::Printf(TEXT(" -run -skipstage -stagingdirectory=\"%s\" -map=%s -device=\"%s\""),
 			*StagePath,
 			*Role->GetInitialMap(),
-			*DeviceProxy->GetDeviceId());
+			*DeviceProxy->GetTargetDeviceId(Variant));
 
 		// cook dependency arguments
 		CommandLine += CookCommand.IsValid() ? CookCommand->GetDependencyArguments(ChainState) : TEXT(" -skipcook");
@@ -80,6 +81,9 @@ private:
 	// Holds the device proxy to launch on.
 	ITargetDeviceProxyPtr DeviceProxy;
 
+	// Holds the name of the flavor of Target Device to use.
+	FName Variant;
+
 	// Holds a pointer to the target platform.
 	const ITargetPlatform& TargetPlatform;
 
@@ -102,8 +106,9 @@ class FLauncherLaunchDedicatedServerCommand
 {
 public:
 
-	FLauncherLaunchDedicatedServerCommand(const ITargetDeviceProxyRef& InDeviceProxy, const ITargetPlatform& InTargetPlatform, const ILauncherProfileLaunchRoleRef& InRole, const TSharedPtr<FLauncherUATCommand>& InCook)
+	FLauncherLaunchDedicatedServerCommand(const ITargetDeviceProxyRef& InDeviceProxy, FName InFlavor, const ITargetPlatform& InTargetPlatform, const ILauncherProfileLaunchRoleRef& InRole, const TSharedPtr<FLauncherUATCommand>& InCook)
 		: DeviceProxy(InDeviceProxy)
+		, Variant(InFlavor)
 		, TargetPlatform(InTargetPlatform)
 		, InstanceId(FGuid::NewGuid())
 		, Role(InRole)
@@ -147,9 +152,9 @@ public:
 		CommandLine += FString::Printf(TEXT(" -run -skipstage -stagingdirectory=\"%s\" -map=%s -device=\"%s\" -dedicatedserver -noclient -serverdevice=\"%s\" -serverplatform=%s"),
 			*StagePath,
 			*Role->GetInitialMap(),
-			*DeviceProxy->GetDeviceId(),
-			*DeviceProxy->GetDeviceId(),
-			*DeviceProxy->GetPlatformName()
+			*DeviceProxy->GetTargetDeviceId(Variant),
+			*DeviceProxy->GetTargetDeviceId(Variant),
+			*DeviceProxy->GetTargetPlatformName(Variant)
 			);
 
 		// requires credentials
@@ -179,6 +184,9 @@ private:
 	// Holds the device proxy to launch on.
 	ITargetDeviceProxyPtr DeviceProxy;
 
+	// Holds the name of the flavor of Target Device to use.
+	FName Variant;
+
 	// Holds a pointer to the target platform.
 	const ITargetPlatform& TargetPlatform;
 
diff --git a/Engine/Source/Developer/LauncherServices/Private/Launcher/LauncherProjectPath.h b/Engine/Source/Developer/LauncherServices/Private/Launcher/LauncherProjectPath.h
new file mode 100644
--- /dev/null
+++ b/Engine/Source/Developer/LauncherServices/Private/Launcher/LauncherProjectPath.h
@@ -0,0 +1,23 @@
+// Copyright 1998-2014 Epic Games, Inc. All Rights Reserved.
+
+#pragma once
+
+
+/**
+ * Implements the launcher project path utilities.
+ */
+class FLauncherProjectPath
+{
+public:
+
+	/**
+	 * Gets the name of the Unreal project to use.
+	 */
+	static FString GetProjectName(const FString & ProjectPath);
+
+	/**
+	 * Gets the base project path for the project (e.g. Samples/Showcases/MyShowcase)
+	 */
+	static FString GetProjectBasePath(const FString & ProjectPath);
+
+};
diff --git a/Engine/Source/Developer/LauncherServices/Private/Launcher/LauncherWorker.cpp b/Engine/Source/Developer/LauncherServices/Private/Launcher/LauncherWorker.cpp
--- a/Engine/Source/Developer/LauncherServices/Private/Launcher/LauncherWorker.cpp
+++ b/Engine/Source/Developer/LauncherServices/Private/Launcher/LauncherWorker.cpp
@@ -100,6 +100,8 @@ uint32 FLauncherWorker::Run( )
 		}		
 	}
 
+	FPlatformProcess::ClosePipe(ReadPipe, WritePipe);
+
 	if (Status == ELauncherWorkerStatus::Canceling)
 	{
 		Status = ELauncherWorkerStatus::Canceled;
@@ -110,9 +112,6 @@ uint32 FLauncherWorker::Run( )
 		LaunchCompleted.Broadcast(TaskChain->Succeeded(), FPlatformTime::Seconds() - LaunchStartTime, TaskChain->ReturnCode());
 	}
 
-	FPlatformProcess::ClosePipe(ReadPipe, WritePipe);
-
-
 	return 0;
 }
 
@@ -206,19 +205,24 @@ void FLauncherWorker::CreateAndExecuteTasks( const ILauncherProfileRef& InProfil
 
 	// determine deployment platforms
 	ILauncherDeviceGroupPtr DeviceGroup = InProfile->GetDeployedDeviceGroup();
+	FName Variant = NAME_None;
 
 	if (DeviceGroup.IsValid() && Platforms.Num() < 1)
 	{
-		const TArray<FString>& Devices = DeviceGroup->GetDevices();
+		const TArray<FString>& Devices = DeviceGroup->GetDeviceIDs();
 		// for each deployed device...
 		for (int32 DeviceIndex = 0; DeviceIndex < Devices.Num(); ++DeviceIndex)
 		{
 			const FString& DeviceId = Devices[DeviceIndex];
 
-			ITargetDeviceProxyRef DeviceProxy = DeviceProxyManager->FindOrAddProxy(DeviceId);
+			ITargetDeviceProxyPtr DeviceProxy = DeviceProxyManager->FindProxyDeviceForTargetDevice(DeviceId);
 
-			// add the platform
-			Platforms.AddUnique(DeviceProxy->GetPlatformName());
+			if (DeviceProxy.IsValid())
+			{
+				// add the platform
+				Variant = DeviceProxy->GetTargetDeviceVariant(DeviceId);
+				Platforms.AddUnique(DeviceProxy->GetTargetPlatformName(Variant));
+			}			
 		}
 	}
 
@@ -373,20 +377,21 @@ void FLauncherWorker::CreateAndExecuteTasks( const ILauncherProfileRef& InProfil
 
 			TMap<FString, TSharedPtr<FLauncherTask> > PerPlatformFileServerTasks;
 
-			const TArray<FString>& Devices = DeviceGroup->GetDevices();
+			const TArray<FString>& Devices = DeviceGroup->GetDeviceIDs();
 
 			// for each deployed device...
 			for (int32 DeviceIndex = 0; DeviceIndex < Devices.Num(); ++DeviceIndex)
 			{
 				const FString& DeviceId = Devices[DeviceIndex];
-
-				ITargetDeviceProxyRef DeviceProxy = DeviceProxyManager->FindOrAddProxy(DeviceId);
-
-				// ... that runs the current platform...
-				if (DeviceProxy->GetPlatformName() != TargetPlatformName)
+				
+				// @Todo: this can probably been cleaned up if the variant class is exposed to the user (not 100% sure i want that yet)
+				//			Deeper investigation / refactoring may make this validation obsolete.
+				ITargetDeviceProxyPtr DeviceProxyPtr = DeviceProxyManager->FindProxyDeviceForTargetDevice(DeviceId);
+				if (!DeviceProxyPtr.IsValid() || DeviceProxyPtr->GetTargetPlatformName(DeviceProxyPtr->GetTargetDeviceVariant(DeviceId)) != TargetPlatformName)
 				{
 					continue;
 				}
+				ITargetDeviceProxyRef DeviceProxy = DeviceProxyPtr.ToSharedRef();
 
 				TSharedPtr<FLauncherTask> PerDeviceTask;
 
@@ -431,11 +436,11 @@ void FLauncherWorker::CreateAndExecuteTasks( const ILauncherProfileRef& InProfil
 					TSharedPtr<FLauncherUATCommand> Command = NULL;
 					if (TargetPlatform->PlatformName() == TEXT("WindowsServer") || TargetPlatform->PlatformName() == TEXT("LinuxServer"))
 					{
-						Command = MakeShareable(new FLauncherDeployServerToDeviceCommand(DeviceProxy, *TargetPlatform, CookCommand));
+						Command = MakeShareable(new FLauncherDeployServerToDeviceCommand(DeviceProxy, Variant, *TargetPlatform, CookCommand));
 					}
 					else
 					{
-						Command = MakeShareable(new FLauncherDeployGameToDeviceCommand(DeviceProxy, *TargetPlatform, CookCommand, LaunchCommandLine));
+						Command = MakeShareable(new FLauncherDeployGameToDeviceCommand(DeviceProxy, Variant, *TargetPlatform, CookCommand, LaunchCommandLine));
 					}
 					TSharedPtr<FLauncherTask> DeployTask = MakeShareable(new FLauncherUATTask(Command, *TargetPlatform, Command->GetName(), ReadPipe, WritePipe, InProfile->GetEditorExe()));
 					DeployTask->OnStarted().AddRaw(this, &FLauncherWorker::OnTaskStarted);
@@ -457,11 +462,11 @@ void FLauncherWorker::CreateAndExecuteTasks( const ILauncherProfileRef& InProfil
 					TSharedPtr<FLauncherUATCommand> Command = NULL;
 					if (TargetPlatform->PlatformName() == TEXT("WindowsServer") || TargetPlatform->PlatformName() == TEXT("LinuxServer"))
 					{
-						Command = MakeShareable(new FLauncherDeployServerPackageToDeviceCommand(DeviceProxy, *TargetPlatform, CookCommand));
+						Command = MakeShareable(new FLauncherDeployServerPackageToDeviceCommand(DeviceProxy, Variant, *TargetPlatform, CookCommand));
 					}
 					else
 					{
-						Command = MakeShareable(new FLauncherDeployGamePackageToDeviceCommand(DeviceProxy, *TargetPlatform, CookCommand, LaunchCommandLine));
+						Command = MakeShareable(new FLauncherDeployGamePackageToDeviceCommand(DeviceProxy, Variant, *TargetPlatform, CookCommand, LaunchCommandLine));
 					}
 					TSharedPtr<FLauncherTask> DeployTask = MakeShareable(new FLauncherUATTask(Command, *TargetPlatform, Command->GetName(), ReadPipe, WritePipe, InProfile->GetEditorExe()));
 					DeployTask->OnStarted().AddRaw(this, &FLauncherWorker::OnTaskStarted);
@@ -480,7 +485,7 @@ void FLauncherWorker::CreateAndExecuteTasks( const ILauncherProfileRef& InProfil
 				}
 				else if (TargetPlatform->PlatformName() == TEXT("XboxOne") || TargetPlatform->PlatformName() == TEXT("IOS") || TargetPlatform->PlatformName().StartsWith(TEXT("Android")))
 				{
-					TSharedPtr<FLauncherUATCommand> Command = MakeShareable(new FLauncherDeployGameToDeviceCommand(DeviceProxy, *TargetPlatform, CookCommand, LaunchCommandLine));
+					TSharedPtr<FLauncherUATCommand> Command = MakeShareable(new FLauncherDeployGameToDeviceCommand(DeviceProxy, Variant, *TargetPlatform, CookCommand, LaunchCommandLine));
 					TSharedPtr<FLauncherTask> DeployTask = MakeShareable(new FLauncherUATTask(Command, *TargetPlatform, Command->GetName(), ReadPipe, WritePipe, InProfile->GetEditorExe()));
 					DeployTask->OnStarted().AddRaw(this, &FLauncherWorker::OnTaskStarted);
 					DeployTask->OnCompleted().AddRaw(this, &FLauncherWorker::OnTaskCompleted);
@@ -510,11 +515,11 @@ void FLauncherWorker::CreateAndExecuteTasks( const ILauncherProfileRef& InProfil
 							TSharedPtr<FLauncherUATCommand> Command;
 							if (Roles[RoleIndex]->GetInstanceType() == ELauncherProfileRoleInstanceTypes::StandaloneClient)
 							{
-								Command = MakeShareable(new FLauncherLaunchGameCommand(DeviceProxy, *TargetPlatform, Roles[RoleIndex].ToSharedRef(), CookCommand));
+								Command = MakeShareable(new FLauncherLaunchGameCommand(DeviceProxy, Variant, *TargetPlatform, Roles[RoleIndex].ToSharedRef(), CookCommand));
 							}
 							else if (Roles[RoleIndex]->GetInstanceType() == ELauncherProfileRoleInstanceTypes::DedicatedServer)
 							{
-								Command = MakeShareable(new FLauncherLaunchDedicatedServerCommand(DeviceProxy, *TargetPlatform, Roles[RoleIndex].ToSharedRef(), CookCommand));
+								Command = MakeShareable(new FLauncherLaunchDedicatedServerCommand(DeviceProxy, Variant, *TargetPlatform, Roles[RoleIndex].ToSharedRef(), CookCommand));
 							}
 								
 							TSharedPtr<FLauncherTask> PerRoleTask = MakeShareable(new FLauncherUATTask(Command, *TargetPlatform, Command->GetName(), ReadPipe, WritePipe, InProfile->GetEditorExe()));
diff --git a/Engine/Source/Developer/LauncherServices/Private/LauncherServicesModule.cpp b/Engine/Source/Developer/LauncherServices/Private/LauncherServicesModule.cpp
--- a/Engine/Source/Developer/LauncherServices/Private/LauncherServicesModule.cpp
+++ b/Engine/Source/Developer/LauncherServices/Private/LauncherServicesModule.cpp
@@ -30,14 +30,17 @@ public:
 
 	virtual ILauncherProfileRef CreateProfile( const FString& ProfileName ) override
 	{
-		return MakeShareable(new FLauncherProfile(ProfileName));
+		ILauncherProfileManagerRef ProfileManager = GetProfileManager();
+		return MakeShareable(new FLauncherProfile(ProfileManager, FGuid(), ProfileName));
 	}
 
 	virtual ILauncherProfileManagerRef GetProfileManager( ) override
 	{
 		if (!ProfileManagerSingleton.IsValid())
 		{
-			ProfileManagerSingleton = MakeShareable(new FLauncherProfileManager());
+			TSharedPtr<FLauncherProfileManager> ProfileManager = MakeShareable(new FLauncherProfileManager());	
+			ProfileManager->Load();
+			ProfileManagerSingleton = ProfileManager;
 		}
 
 		return ProfileManagerSingleton.ToSharedRef();
diff --git a/Engine/Source/Developer/LauncherServices/Private/LauncherServicesPrivatePCH.h b/Engine/Source/Developer/LauncherServices/Private/LauncherServicesPrivatePCH.h
--- a/Engine/Source/Developer/LauncherServices/Private/LauncherServicesPrivatePCH.h
+++ b/Engine/Source/Developer/LauncherServices/Private/LauncherServicesPrivatePCH.h
@@ -20,11 +20,15 @@
 /* Private includes
  *****************************************************************************/
 
+/** Defines the launcher simple profile file version. */
+#define LAUNCHERSERVICES_SIMPLEPROFILEVERSION 1
+
 /** Defines the launcher profile file version. */
-#define LAUNCHERSERVICES_PROFILEVERSION 8
+#define LAUNCHERSERVICES_PROFILEVERSION 10
 
 
 // profile manager
+#include "LauncherProjectPath.h"
 #include "LauncherDeviceGroup.h"
 #include "LauncherProfileLaunchRole.h"
 #include "LauncherProfile.h"
diff --git a/Engine/Source/Developer/LauncherServices/Private/Profiles/LauncherDeviceGroup.h b/Engine/Source/Developer/LauncherServices/Private/Profiles/LauncherDeviceGroup.h
--- a/Engine/Source/Developer/LauncherServices/Private/Profiles/LauncherDeviceGroup.h
+++ b/Engine/Source/Developer/LauncherServices/Private/Profiles/LauncherDeviceGroup.h
@@ -15,7 +15,7 @@ public:
 	/**
 	 * Default constructor.
 	 */
-	FLauncherDeviceGroup( ) { }
+	FLauncherDeviceGroup() { }
 
 	/**
 	 * Creates and initializes a new instance.
@@ -23,7 +23,7 @@ public:
 	 * @param InId - The unique group identifier.
 	 * @param InName - The name of the group.
 	 */
-	FLauncherDeviceGroup( const FGuid& InId, const FString& InName )
+	FLauncherDeviceGroup(const FGuid& InId, const FString& InName)
 		: Id(InId)
 		, Name(InName)
 	{ }
@@ -32,51 +32,51 @@ public:
 
 	// Begin ILauncherDeviceGroup interface
 
-	virtual void AddDevice( const FString& Device ) override
+	virtual void AddDevice(const FString& DeviceID) override
 	{
-		Devices.AddUnique(Device);
+		Devices.AddUnique(DeviceID);
 
-		DeviceAddedDelegate.Broadcast(AsShared(), Device);
+		DeviceAddedDelegate.Broadcast(AsShared(), DeviceID);
 	}
 
-	virtual const TArray<FString>& GetDevices( ) override
+	virtual const TArray<FString>& GetDeviceIDs() override
 	{
 		return Devices;
 	}
 
-	virtual FGuid GetId( ) const override
+	virtual FGuid GetId() const override
 	{
 		return Id;
 	}
 
-	virtual const FString& GetName( ) const override
+	virtual const FString& GetName() const override
 	{
 		return Name;
 	}
 
-	virtual int32 GetNumDevices( ) const override
+	virtual int32 GetNumDevices() const override
 	{
 		return Devices.Num();
 	}
 
-	virtual FOnLauncherDeviceGroupDeviceAdded& OnDeviceAdded( ) override
+	virtual FOnLauncherDeviceGroupDeviceAdded& OnDeviceAdded() override
 	{
 		return DeviceAddedDelegate;
 	}
 
-	virtual FOnLauncherDeviceGroupDeviceRemoved& OnDeviceRemoved( ) override
+	virtual FOnLauncherDeviceGroupDeviceRemoved& OnDeviceRemoved() override
 	{
 		return DeviceRemovedDelegate;
 	}
 
-	virtual void RemoveDevice( const FString& Device ) override
+	virtual void RemoveDevice(const FString& DeviceID) override
 	{
-		Devices.Remove(Device);
+		Devices.Remove(DeviceID);
 
-		DeviceRemovedDelegate.Broadcast(AsShared(), Device);
+		DeviceRemovedDelegate.Broadcast(AsShared(), DeviceID);
 	}
 
-	virtual void SetName( const FString& NewName ) override
+	virtual void SetName(const FString& NewName) override
 	{
 		Name = NewName;
 	}
diff --git a/Engine/Source/Developer/LauncherServices/Private/Profiles/LauncherProfile.h b/Engine/Source/Developer/LauncherServices/Private/Profiles/LauncherProfile.h
--- a/Engine/Source/Developer/LauncherServices/Private/Profiles/LauncherProfile.h
+++ b/Engine/Source/Developer/LauncherServices/Private/Profiles/LauncherProfile.h
@@ -2,6 +2,112 @@
 
 #pragma once
 
+/**
+* Implements a simple profile which controls the desired output of the Launcher for simple
+*/
+class FLauncherSimpleProfile
+	: public ILauncherSimpleProfile
+{
+public:
+
+	FLauncherSimpleProfile(const FString& InDeviceName)
+		: DeviceName(InDeviceName)
+	{
+		SetDefaults();
+	}
+
+	// Begin ILauncherSimpleProfile interface
+
+	virtual const FString& GetDeviceName() const override
+	{
+		return DeviceName;
+	}
+
+	virtual FName GetDeviceVariant() const  override
+	{
+		return Variant;
+	}
+
+	virtual EBuildConfigurations::Type GetBuildConfiguration() const override
+	{
+		return BuildConfiguration;
+	}
+
+	virtual ELauncherProfileCookModes::Type GetCookMode() const override
+	{
+		return CookMode;
+	}
+
+	virtual void SetDeviceName(const FString& InDeviceName) override
+	{
+		if (DeviceName != InDeviceName)
+		{
+			DeviceName = InDeviceName;
+		}
+	}
+
+	virtual void SetDeviceVariant(FName InVariant) override
+	{
+		Variant = InVariant;
+	}
+
+	virtual void SetBuildConfiguration(EBuildConfigurations::Type InConfiguration) override
+	{
+		BuildConfiguration = InConfiguration;
+	}
+
+	virtual void SetCookMode(ELauncherProfileCookModes::Type InMode) override
+	{
+		CookMode = InMode;
+	}
+
+	virtual bool Serialize(FArchive& Archive) override
+	{
+		int32 Version = LAUNCHERSERVICES_SIMPLEPROFILEVERSION;
+
+		Archive << Version;
+
+		if (Version != LAUNCHERSERVICES_SIMPLEPROFILEVERSION)
+		{
+			return false;
+		}
+
+		// IMPORTANT: make sure to bump LAUNCHERSERVICES_SIMPLEPROFILEVERSION when modifying this!
+		Archive << DeviceName
+			<< Variant
+			<< BuildConfiguration
+			<< CookMode;
+
+		return true;
+	}
+
+	virtual void SetDefaults() override
+	{
+		// None will mean the preferred Variant for the device is used.
+		Variant = NAME_None;
+		
+		// I don't use FApp::GetBuildConfiguration() because i don't want the act of running in debug the first time to cause the simple
+		// profiles created for your persistent devices to be in debug. The use might not see this if they don't expand the Advanced options.
+		BuildConfiguration = EBuildConfigurations::Development;
+		
+		CookMode = ELauncherProfileCookModes::OnTheFly;
+	}
+
+private:
+
+	// Holds the name of the device this simple profile is for
+	FString DeviceName;
+
+	// Holds the name of the device variant.
+	FName Variant;
+
+	// Holds the desired build configuration (only used if creating new builds).
+	TEnumAsByte<EBuildConfigurations::Type> BuildConfiguration;
+
+	// Holds the cooking mode.
+	TEnumAsByte<ELauncherProfileCookModes::Type> CookMode;
+};
+
 
 /**
  * Implements a profile which controls the desired output of the Launcher
@@ -14,8 +120,9 @@ public:
 	/**
 	 * Default constructor.
 	 */
-	FLauncherProfile( )
-		: DefaultLaunchRole(MakeShareable(new FLauncherProfileLaunchRole()))
+	FLauncherProfile(ILauncherProfileManagerRef ProfileManager)
+		: LauncherProfileManager(ProfileManager)
+		, DefaultLaunchRole(MakeShareable(new FLauncherProfileLaunchRole()))
 	{ }
 
 	/**
@@ -23,9 +130,10 @@ public:
 	 *
 	 * @param InProfileName - The name of the profile.
 	 */
-	FLauncherProfile( const FString& InProfileName )
-		: DefaultLaunchRole(MakeShareable(new FLauncherProfileLaunchRole()))
-		, Id(FGuid::NewGuid())
+	FLauncherProfile(ILauncherProfileManagerRef ProfileManager, FGuid InId, const FString& InProfileName)
+		: LauncherProfileManager(ProfileManager)
+		, DefaultLaunchRole(MakeShareable(new FLauncherProfileLaunchRole()))
+		, Id(InId)
 		, Name(InProfileName)
 	{
 		SetDefaults();
@@ -217,6 +325,11 @@ public:
 		return Name;
 	}
 
+	virtual FString GetDescription() const override
+	{
+		return Description;
+	}
+
 	virtual ELauncherProfilePackagingModes::Type GetPackagingMode( ) const override
 	{
 		return PackagingMode;
@@ -227,37 +340,48 @@ public:
 		return PackageDir;
 	}
 
-	virtual FString GetProjectName( ) const override
+	virtual bool HasProjectSpecified() const override
 	{
-		if (!ProjectPath.IsEmpty())
+		return ProjectSpecified;
+	}
+
+	virtual FString GetProjectName() const override
+	{
+		if (ProjectSpecified)
 		{
-			return FPaths::GetBaseFilename(ProjectPath);
+			FString Path = FLauncherProjectPath::GetProjectName(ProjectPath);
+			return Path;
 		}
-
-		return FString();
+		return LauncherProfileManager->GetProjectName();
 	}
 
 	virtual FString GetProjectBasePath() const override
 	{
-		if (!ProjectPath.IsEmpty())
+		if (ProjectSpecified)
 		{
-			FString OutPath = ProjectPath;
-			FPaths::MakePathRelativeTo(OutPath, *FPaths::RootDir());
-			return FPaths::GetPath(OutPath);
+			return FLauncherProjectPath::GetProjectBasePath(ProjectPath);
 		}
-
-		return FString();
+		return LauncherProfileManager->GetProjectBasePath();
 	}
 
-	virtual const FString& GetProjectPath( ) const override
+	virtual FString GetProjectPath() const override
 	{
-		return ProjectPath;
+		if (ProjectSpecified)
+		{
+			return ProjectPath;
+		}
+		return LauncherProfileManager->GetProjectPath();
 	}
 
     virtual uint32 GetTimeout() const override
     {
         return Timeout;
     }
+
+	virtual bool HasValidationError() const override
+	{
+		return ValidationErrors.Num() > 0;
+	}
     
 	virtual bool HasValidationError( ELauncherProfileValidationErrors::Type Error ) const override
 	{
@@ -368,7 +492,9 @@ public:
 		// IMPORTANT: make sure to bump LAUNCHERSERVICES_PROFILEVERSION when modifying this!
 		Archive << Id
 				<< Name
+				<< Description
 				<< BuildConfiguration
+				<< ProjectSpecified
 				<< ProjectPath
 				<< CookConfiguration
 				<< CookIncremental
@@ -422,6 +548,8 @@ public:
 
 	virtual void SetDefaults( ) override
 	{
+		ProjectSpecified = false;
+
 		// default project settings
 		if (FPaths::IsProjectFilePathSet())
 		{
@@ -444,7 +572,7 @@ public:
 
 		// default cook settings
 		CookConfiguration = FApp::GetBuildConfiguration();
-		CookMode = ELauncherProfileCookModes::DoNotCook;
+		CookMode = ELauncherProfileCookModes::OnTheFly;
 		CookOptions = FString();
 		CookIncremental = false;
 		CookUnversioned = false;
@@ -452,8 +580,8 @@ public:
 		CookedCultures.Add(I18N.GetCurrentCulture()->GetName());
 		CookedMaps.Reset();
 		CookedPlatforms.Reset();
-        ForceClose = false;
-        Timeout = 180;
+        ForceClose = true;
+        Timeout = 60;
 
 /*		if (GetTargetPlatformManager()->GetRunningTargetPlatform() != NULL)
 		{
@@ -462,14 +590,14 @@ public:
 
 		// default deploy settings
 		DeployedDeviceGroup.Reset();
-		DeploymentMode = ELauncherProfileDeploymentModes::DoNotDeploy;
+		DeploymentMode = ELauncherProfileDeploymentModes::CopyToDevice;
 		DeployStreamingServer = false;
 		DeployWithUnrealPak = false;
 		DeployedDeviceGroupId = FGuid();
 		HideFileServerWindow = false;
 
 		// default launch settings
-		LaunchMode = ELauncherProfileLaunchModes::DoNotLaunch;
+		LaunchMode = ELauncherProfileLaunchModes::DefaultRole;
 		DefaultLaunchRole->SetCommandLine(FString());
 		DefaultLaunchRole->SetInitialCulture(I18N.GetCurrentCulture()->GetName());
 		DefaultLaunchRole->SetInitialMap(FString());
@@ -527,6 +655,16 @@ public:
 		}
 	}
 
+	virtual void SetCookOptions(const FString& Options) override
+	{
+		if (CookOptions != Options)
+		{
+			CookOptions = Options;
+
+			Validate();
+		}
+	}
+
 	virtual void SetDeployWithUnrealPak( bool UseUnrealPak ) override
 	{
 		if (DeployWithUnrealPak != UseUnrealPak)
@@ -613,6 +751,16 @@ public:
 		}
 	}
 
+	virtual void SetDescription(const FString& NewDescription) override
+	{
+		if (Description != NewDescription)
+		{
+			Description = NewDescription;
+
+			Validate();
+		}
+	}
+
 	virtual void SetPackagingMode( ELauncherProfilePackagingModes::Type Mode ) override
 	{
 		if (PackagingMode != Mode)
@@ -633,11 +781,40 @@ public:
 		}
 	}
 
+	virtual void SetProjectSpecified(bool Specified) override
+	{
+		if (ProjectSpecified != Specified)
+		{
+			ProjectSpecified = Specified;
+
+			Validate();
+
+			ProjectChangedDelegate.Broadcast();
+		}
+	}
+
+	virtual void FallbackProjectUpdated() override
+	{
+		if (!HasProjectSpecified())
+		{
+			Validate();
+
+			ProjectChangedDelegate.Broadcast();
+		}
+	}
+
 	virtual void SetProjectPath( const FString& Path ) override
 	{
 		if (ProjectPath != Path)
 		{
-			ProjectPath = FPaths::ConvertRelativePathToFull(Path);
+			if(Path.IsEmpty())
+			{
+				ProjectPath = Path;
+			}
+			else
+			{
+				ProjectPath = FPaths::ConvertRelativePathToFull(Path);
+			}
 			CookedMaps.Reset();
 
 			Validate();
@@ -718,7 +895,7 @@ protected:
 		}
 
 		// Build: a project must be selected
-		if (ProjectPath.IsEmpty())
+		if (GetProjectPath().IsEmpty())
 		{
 			ValidationErrors.Add(ELauncherProfileValidationErrors::NoProjectSelected);
 		}
@@ -815,7 +992,7 @@ protected:
 		// Deploy: ensure that all the target device SDKs are installed
 		if ((DeploymentMode != ELauncherProfileDeploymentModes::DoNotDeploy) && DeployedDeviceGroup.IsValid())
 		{
-			const TArray<FString>& Devices = DeployedDeviceGroup->GetDevices();
+			const TArray<FString>& Devices = DeployedDeviceGroup->GetDeviceIDs();
 			for(auto DeviceId : Devices)
 			{
 				TSharedPtr<ITargetDeviceServicesModule> TargetDeviceServicesModule = StaticCastSharedPtr<ITargetDeviceServicesModule>(FModuleManager::Get().LoadModule(TEXT("TargetDeviceServices")));
@@ -826,7 +1003,7 @@ protected:
 					
 					if(DeviceProxy.IsValid())
 					{
-						FString const& PlatformName = DeviceProxy->GetPlatformName();
+						FString const& PlatformName = DeviceProxy->GetTargetPlatformName(DeviceProxy->GetTargetDeviceVariant(DeviceId));
 						bool bProjectHasCode = false; // @todo: Does the project have any code?
 						FString NotInstalledDocLink;
 						const ITargetPlatform* Platform = GetTargetPlatformManager()->FindTargetPlatform(PlatformName);
@@ -862,6 +1039,9 @@ protected:
 
 private:
 
+	//  Holds a reference to the launcher profile manager.
+	ILauncherProfileManagerRef LauncherProfileManager;
+
 	// Holds the desired build configuration (only used if creating new builds).
 	TEnumAsByte<EBuildConfigurations::Type> BuildConfiguration;
 
@@ -926,12 +1106,18 @@ private:
 	// Holds the profile's name.
 	FString Name;
 
+	// Holds the profile's description.
+	FString Description;
+
 	// Holds the packaging mode.
 	TEnumAsByte<ELauncherProfilePackagingModes::Type> PackagingMode;
 
-	// Holda the package directory
+	// Holds the package directory
 	FString PackageDir;
 
+	// Holds a flag indicating whether the project is specified by this profile.
+	bool ProjectSpecified;
+
 	// Holds the path to the Unreal project used by this profile.
 	FString ProjectPath;
 
diff --git a/Engine/Source/Developer/LauncherServices/Private/Profiles/LauncherProfileManager.cpp b/Engine/Source/Developer/LauncherServices/Private/Profiles/LauncherProfileManager.cpp
--- a/Engine/Source/Developer/LauncherServices/Private/Profiles/LauncherProfileManager.cpp
+++ b/Engine/Source/Developer/LauncherServices/Private/Profiles/LauncherProfileManager.cpp
@@ -6,16 +6,19 @@
 /* ILauncherProfileManager structors
  *****************************************************************************/
 
-FLauncherProfileManager::FLauncherProfileManager( )
+FLauncherProfileManager::FLauncherProfileManager()
 {
-	LoadDeviceGroups();
-	LoadProfiles();
 }
 
-
 /* ILauncherProfileManager interface
  *****************************************************************************/
 
+void FLauncherProfileManager::Load()
+{
+	LoadDeviceGroups();
+	LoadProfiles();
+}
+
 void FLauncherProfileManager::AddDeviceGroup( const ILauncherDeviceGroupRef& DeviceGroup )
 {
 	if (!DeviceGroups.Contains(DeviceGroup))
@@ -47,15 +50,51 @@ ILauncherDeviceGroupRef FLauncherProfileManager::AddNewDeviceGroup( )
 }
 
 
-ILauncherProfileRef FLauncherProfileManager::AddNewProfile( )
+ILauncherDeviceGroupRef FLauncherProfileManager::CreateUnmanagedDeviceGroup()
+{
+	ILauncherDeviceGroupRef  NewGroup = MakeShareable(new FLauncherDeviceGroup(FGuid::NewGuid(), TEXT("Simple Group")));
+	return NewGroup;
+}
+
+
+ILauncherSimpleProfilePtr FLauncherProfileManager::FindOrAddSimpleProfile(const FString& DeviceName)
+{
+	// replace the existing profile
+	ILauncherSimpleProfilePtr SimpleProfile = FindSimpleProfile(DeviceName);
+	if (!SimpleProfile.IsValid())
+	{
+		SimpleProfile = MakeShareable(new FLauncherSimpleProfile(DeviceName));
+		SimpleProfiles.Add(SimpleProfile);
+	}
+	
+	return SimpleProfile;
+}
+
+
+ILauncherSimpleProfilePtr FLauncherProfileManager::FindSimpleProfile(const FString& DeviceName)
+{
+	for (int32 ProfileIndex = 0; ProfileIndex < SimpleProfiles.Num(); ++ProfileIndex)
+	{
+		ILauncherSimpleProfilePtr SimpleProfile = SimpleProfiles[ProfileIndex];
+
+		if (SimpleProfile->GetDeviceName() == DeviceName)
+		{
+			return SimpleProfile;
+		}
+	}
+
+	return nullptr;
+}
+
+ILauncherProfileRef FLauncherProfileManager::AddNewProfile()
 {
 	// find a unique name for the profile.
-	int32 ProfileIndex = Profiles.Num();
+	int32 ProfileIndex = SavedProfiles.Num();
 	FString ProfileName = FString::Printf(TEXT("New Profile %d"), ProfileIndex);
 
-	for (int32 Index = 0; Index < Profiles.Num(); ++Index)
+	for (int32 Index = 0; Index < SavedProfiles.Num(); ++Index)
 	{
-		if (Profiles[Index]->GetName() == ProfileName)
+		if (SavedProfiles[Index]->GetName() == ProfileName)
 		{
 			ProfileName = FString::Printf(TEXT("New Profile %d"), ++ProfileIndex);
 			Index = -1;
@@ -65,17 +104,29 @@ ILauncherProfileRef FLauncherProfileManager::AddNewProfile( )
 	}
 
 	// create and add the profile
-	ILauncherProfileRef NewProfile = MakeShareable(new FLauncherProfile(ProfileName));
+	ILauncherProfileRef NewProfile = MakeShareable(new FLauncherProfile(AsShared(), FGuid::NewGuid(), ProfileName));
 
 	AddProfile(NewProfile);
 
+	SaveProfile(NewProfile);
+
+	return NewProfile;
+}
+
+ILauncherProfileRef FLauncherProfileManager::CreateUnsavedProfile(FString ProfileName)
+{
+	// create and return the profile
+	ILauncherProfileRef NewProfile = MakeShareable(new FLauncherProfile(AsShared(), FGuid(), ProfileName));
+	
+	AllProfiles.Add(NewProfile);
+	
 	return NewProfile;
 }
 
 
 void FLauncherProfileManager::AddProfile( const ILauncherProfileRef& Profile )
 {
-	if (!Profiles.Contains(Profile))
+	if (!SavedProfiles.Contains(Profile))
 	{
 		// replace the existing profile
 		ILauncherProfilePtr ExistingProfile = GetProfile(Profile->GetId());
@@ -85,9 +136,14 @@ void FLauncherProfileManager::AddProfile( const ILauncherProfileRef& Profile )
 			RemoveProfile(ExistingProfile.ToSharedRef());
 		}
 
+		if (!Profile->GetDeployedDeviceGroup().IsValid())
+		{
+			Profile->SetDeployedDeviceGroup(AddNewDeviceGroup());
+		}
+
 		// add the new profile
-		Profiles.Add(Profile);
-		SaveProfiles();
+		SavedProfiles.Add(Profile);
+		AllProfiles.Add(Profile);
 
 		ProfileAddedDelegate.Broadcast(Profile);
 	}
@@ -96,9 +152,9 @@ void FLauncherProfileManager::AddProfile( const ILauncherProfileRef& Profile )
 
 ILauncherProfilePtr FLauncherProfileManager::FindProfile( const FString& ProfileName )
 {
-	for (int32 ProfileIndex = 0; ProfileIndex < Profiles.Num(); ++ProfileIndex)
+	for (int32 ProfileIndex = 0; ProfileIndex < SavedProfiles.Num(); ++ProfileIndex)
 	{
-		ILauncherProfilePtr Profile = Profiles[ProfileIndex];
+		ILauncherProfilePtr Profile = SavedProfiles[ProfileIndex];
 
 		if (Profile->GetName() == ProfileName)
 		{
@@ -118,7 +174,7 @@ const TArray<ILauncherDeviceGroupPtr>& FLauncherProfileManager::GetAllDeviceGrou
 
 const TArray<ILauncherProfilePtr>& FLauncherProfileManager::GetAllProfiles( ) const
 {
-	return Profiles;
+	return SavedProfiles;
 }
 
 
@@ -140,9 +196,9 @@ ILauncherDeviceGroupPtr FLauncherProfileManager::GetDeviceGroup( const FGuid& Gr
 
 ILauncherProfilePtr FLauncherProfileManager::GetProfile( const FGuid& ProfileId ) const
 {
-	for (int32 ProfileIndex = 0; ProfileIndex < Profiles.Num(); ++ProfileIndex)
+	for (int32 ProfileIndex = 0; ProfileIndex < SavedProfiles.Num(); ++ProfileIndex)
 	{
-		ILauncherProfilePtr Profile = Profiles[ProfileIndex];
+		ILauncherProfilePtr Profile = SavedProfiles[ProfileIndex];
 
 		if (Profile->GetId() == ProfileId)
 		{
@@ -156,11 +212,16 @@ ILauncherProfilePtr FLauncherProfileManager::GetProfile( const FGuid& ProfileId
 
 ILauncherProfilePtr FLauncherProfileManager::LoadProfile( FArchive& Archive )
 {
-	FLauncherProfile* Profile = new FLauncherProfile();
+	FLauncherProfile* Profile = new FLauncherProfile(AsShared());
 
 	if (Profile->Serialize(Archive))
 	{
-		Profile->SetDeployedDeviceGroup(GetDeviceGroup(Profile->GetDeployedDeviceGroupId()));
+		ILauncherDeviceGroupPtr DeviceGroup = GetDeviceGroup(Profile->GetDeployedDeviceGroupId());
+		if (!DeviceGroup.IsValid())
+		{
+			DeviceGroup = AddNewDeviceGroup();	
+		}
+		Profile->SetDeployedDeviceGroup(DeviceGroup);
 
 		return MakeShareable(Profile);
 	}
@@ -187,37 +248,89 @@ void FLauncherProfileManager::RemoveDeviceGroup( const ILauncherDeviceGroupRef&
 }
 
 
+void FLauncherProfileManager::RemoveSimpleProfile(const ILauncherSimpleProfileRef& SimpleProfile)
+{
+	if (SimpleProfiles.Remove(SimpleProfile) > 0)
+	{
+		// delete the persisted simple profile on disk
+		FString SimpleProfileFileName = GetProfileFolder() / SimpleProfile->GetDeviceName() + TEXT(".uslp");
+		IFileManager::Get().Delete(*SimpleProfileFileName);
+	}
+}
+
+
 void FLauncherProfileManager::RemoveProfile( const ILauncherProfileRef& Profile )
 {
-	if (Profiles.Remove(Profile) > 0)
+	AllProfiles.Remove(Profile);
+	if (SavedProfiles.Remove(Profile) > 0)
 	{
-		// delete the persisted profile on disk
-		FString ProfileFileName = GetProfileFolder() / Profile->GetId().ToString() + TEXT(".ulp");
+		if (Profile->GetId().IsValid())
+		{
+			// delete the persisted profile on disk
+			FString ProfileFileName = GetProfileFolder() / Profile->GetId().ToString() + TEXT(".ulp");
 
-		// delete the profile
-		IFileManager::Get().Delete(*ProfileFileName);
-		SaveProfiles();
+			// delete the profile
+			IFileManager::Get().Delete(*ProfileFileName);
 
-		ProfileRemovedDelegate.Broadcast(Profile);
+			ProfileRemovedDelegate.Broadcast(Profile);
+		}
 	}
 }
 
 
-void FLauncherProfileManager::SaveProfile( const ILauncherProfileRef& Profile, FArchive& Archive )
+void FLauncherProfileManager::SaveProfile(const ILauncherProfileRef& Profile)
 {
-	Profile->Serialize(Archive);
+	if (Profile->GetId().IsValid())
+	{
+		FString ProfileFileName = GetProfileFolder() / Profile->GetId().ToString() + TEXT(".ulp");
+		FArchive* ProfileFileWriter = IFileManager::Get().CreateFileWriter(*ProfileFileName);
+
+		if (ProfileFileWriter != nullptr)
+		{
+			Profile->Serialize(*ProfileFileWriter);
+
+			delete ProfileFileWriter;
+		}
+	}
 }
 
 
 void FLauncherProfileManager::SaveSettings( )
 {
 	SaveDeviceGroups();
+	SaveSimpleProfiles();
 	SaveProfiles();
 }
 
+FString FLauncherProfileManager::GetProjectName() const
+{
+	return FLauncherProjectPath::GetProjectName(ProjectPath);
+}
 
-/* FLauncherProfileManager implementation
- *****************************************************************************/
+FString FLauncherProfileManager::GetProjectBasePath() const
+{
+	return FLauncherProjectPath::GetProjectBasePath(ProjectPath);
+}
+
+FString FLauncherProfileManager::GetProjectPath() const
+{
+	return ProjectPath;
+}
+
+void FLauncherProfileManager::SetProjectPath(const FString& InProjectPath)
+{
+	if (ProjectPath != InProjectPath)
+	{
+		ProjectPath = InProjectPath;
+		for (ILauncherProfilePtr Profile : AllProfiles)
+		{
+			if (Profile.IsValid())
+			{
+				Profile->FallbackProjectUpdated();
+			}
+		}
+	}
+}
 
 void FLauncherProfileManager::LoadDeviceGroups( )
 {
@@ -317,7 +430,7 @@ void FLauncherProfileManager::SaveDeviceGroups( )
 		for (int32 GroupIndex = 0; GroupIndex < DeviceGroups.Num(); ++GroupIndex)
 		{
 			const ILauncherDeviceGroupPtr& Group = DeviceGroups[GroupIndex];
-			const TArray<FString>& Devices = Group->GetDevices();
+			const TArray<FString>& Devices = Group->GetDeviceIDs();
 
 			FString DeviceListString;
 
@@ -343,18 +456,27 @@ void FLauncherProfileManager::SaveDeviceGroups( )
 }
 
 
-void FLauncherProfileManager::SaveProfiles( )
+void FLauncherProfileManager::SaveSimpleProfiles()
 {
-	for (TArray<ILauncherProfilePtr>::TIterator It(Profiles); It; ++It)
+	for (TArray<ILauncherSimpleProfilePtr>::TIterator It(SimpleProfiles); It; ++It)
 	{
-		FString ProfileFileName = GetProfileFolder() / (*It)->GetId().ToString() + TEXT(".ulp");
-		FArchive* ProfileFileWriter = IFileManager::Get().CreateFileWriter(*ProfileFileName);
+		FString SimpleProfileFileName = GetProfileFolder() / (*It)->GetDeviceName() + TEXT(".uslp");
+		FArchive* ProfileFileWriter = IFileManager::Get().CreateFileWriter(*SimpleProfileFileName);
 
 		if (ProfileFileWriter != nullptr)
 		{
-			SaveProfile((*It).ToSharedRef(), *ProfileFileWriter);
+			(*It)->Serialize(*ProfileFileWriter);
 
 			delete ProfileFileWriter;
 		}
 	}
 }
+
+
+void FLauncherProfileManager::SaveProfiles( )
+{
+	for (TArray<ILauncherProfilePtr>::TIterator It(SavedProfiles); It; ++It)
+	{
+		SaveProfile((*It).ToSharedRef());
+	}
+}
diff --git a/Engine/Source/Developer/LauncherServices/Private/Profiles/LauncherProfileManager.h b/Engine/Source/Developer/LauncherServices/Private/Profiles/LauncherProfileManager.h
--- a/Engine/Source/Developer/LauncherServices/Private/Profiles/LauncherProfileManager.h
+++ b/Engine/Source/Developer/LauncherServices/Private/Profiles/LauncherProfileManager.h
@@ -7,23 +7,37 @@
  * Implements a helper class that manages all profiles in the Launcher
  */
 class FLauncherProfileManager
-	: public ILauncherProfileManager
+	: public TSharedFromThis<FLauncherProfileManager>
+	, public ILauncherProfileManager
 {
 public:
 
 	/** Default constructor. */
-	FLauncherProfileManager( );
+	FLauncherProfileManager();
 
 
 public:
 
+	/**
+	 * Loads the profiles and device groups.
+	 */
+	void Load();
+
 	// Begin ILauncherProfileManager interface
 
 	virtual void AddDeviceGroup( const ILauncherDeviceGroupRef& DeviceGroup ) override;
 
 	virtual ILauncherDeviceGroupRef AddNewDeviceGroup( ) override;
 
-	virtual ILauncherProfileRef AddNewProfile( ) override;
+	virtual ILauncherDeviceGroupRef CreateUnmanagedDeviceGroup() override;
+
+	virtual ILauncherSimpleProfilePtr FindOrAddSimpleProfile(const FString& DeviceName) override;
+
+	virtual ILauncherSimpleProfilePtr FindSimpleProfile(const FString& DeviceName) override;
+
+	virtual ILauncherProfileRef AddNewProfile() override;
+
+	virtual ILauncherProfileRef CreateUnsavedProfile(FString ProfileName) override;
 
 	virtual void AddProfile( const ILauncherProfileRef& Profile ) override;
 
@@ -63,12 +77,22 @@ public:
 
 	virtual void RemoveDeviceGroup( const ILauncherDeviceGroupRef& DeviceGroup ) override;
 
+	virtual void RemoveSimpleProfile(const ILauncherSimpleProfileRef& SimpleProfile) override;
+
 	virtual void RemoveProfile( const ILauncherProfileRef& Profile ) override;
 
-	virtual void SaveProfile( const ILauncherProfileRef& Profile, FArchive& Archive ) override;
+	virtual void SaveProfile( const ILauncherProfileRef& Profile) override;
 
 	virtual void SaveSettings( ) override;
 
+	virtual FString GetProjectName() const override;
+
+	virtual FString GetProjectBasePath() const override;
+
+	virtual FString GetProjectPath() const override;
+
+	virtual void SetProjectPath(const FString& InProjectPath) override;
+
 	// End ILauncherProfileManager interface
 
 protected:
@@ -95,12 +119,17 @@ protected:
 	/**
 	 * Saves all the device groups to a config file
 	 */
-	void SaveDeviceGroups( );
+	void SaveDeviceGroups();
+
+	/*
+	* Saves all simple profiles to disk.
+	*/
+	void SaveSimpleProfiles();
 
 	/**
 	 * Saves all profiles to disk.
 	 */
-	void SaveProfiles( );
+	void SaveProfiles();
 
 protected:
 
@@ -119,8 +148,17 @@ private:
 	// Holds the collection of device groups.
 	TArray<ILauncherDeviceGroupPtr> DeviceGroups;
 
+	// Holds the collection of simple launcher profiles.
+	TArray<ILauncherSimpleProfilePtr> SimpleProfiles;
+
 	// Holds the collection of launcher profiles.
-	TArray<ILauncherProfilePtr> Profiles;
+	TArray<ILauncherProfilePtr> SavedProfiles;
+
+	// Holds the collection of launcher profiles.
+	TArray<ILauncherProfilePtr> AllProfiles;
+
+	// Holds the currently selected project path
+	FString ProjectPath;
 
 private:
 
@@ -135,4 +173,5 @@ private:
 
 	// Holds a delegate to be invoked when a launcher profile was removed.
 	FOnLauncherProfileManagerProfileRemoved ProfileRemovedDelegate;
+
 };
diff --git a/Engine/Source/Developer/LauncherServices/Public/Interfaces/ILauncherDeviceGroup.h b/Engine/Source/Developer/LauncherServices/Public/Interfaces/ILauncherDeviceGroup.h
--- a/Engine/Source/Developer/LauncherServices/Public/Interfaces/ILauncherDeviceGroup.h
+++ b/Engine/Source/Developer/LauncherServices/Public/Interfaces/ILauncherDeviceGroup.h
@@ -37,51 +37,51 @@ public:
 	/**
 	 * Adds a device to the group.
 	 *
-	 * @param DeviceId - The identifier of the device to add.
+	 * @param DeviceID - The identifier of the device to add.
 	 */
-	virtual void AddDevice( const FString& DeviceId ) = 0;
+	virtual void AddDevice(const FString& DeviceID) = 0;
 	
 	/**
 	 * Get a list of identifiers for devices associated with this device group.
 	 *
 	 * @return The list of device identifiers.
 	 */
-	virtual const TArray<FString>& GetDevices( ) = 0;
+	virtual const TArray<FString>& GetDeviceIDs() = 0;
 
 	/**
 	 * Gets the unique identifier of the device group.
 	 *
 	 * @return The device group identifier.
 	 */
-	virtual FGuid GetId( ) const = 0;
+	virtual FGuid GetId() const = 0;
 
 	/**
 	 * Gets the human readable name of the device group.
 	 *
 	 * @return The name of the device group.
 	 */
-	virtual const FString& GetName( ) const = 0;
+	virtual const FString& GetName() const = 0;
 
 	/**
 	 * Get a list of devices associated with this device group.
 	 *
 	 * @return The number of devices the device group maintains.
 	 */
-	virtual int32 GetNumDevices( ) const = 0;
+	virtual int32 GetNumDevices() const = 0;
 
 	/**
 	 * Removes a device from the group.
 	 *
-	 * @param DeviceId - The identifier of the device to remove.
+	 * @param DeviceID - The identifier of the device to remove.
 	 */
-	virtual void RemoveDevice( const FString& DeviceId ) = 0;
+	virtual void RemoveDevice(const FString& DeviceID) = 0;
 
 	/**
 	 * Sets the name of the device group.
 	 *
 	 * @param NewName - The new name.
 	 */
-	virtual void SetName( const FString& NewName ) = 0;
+	virtual void SetName(const FString& NewName) = 0;
 
 public:
 
@@ -90,19 +90,19 @@ public:
 	 *
 	 * @return The delegate.
 	 */
-	virtual FOnLauncherDeviceGroupDeviceAdded& OnDeviceAdded( ) = 0;
+	virtual FOnLauncherDeviceGroupDeviceAdded& OnDeviceAdded() = 0;
 
 	/**
 	 * Returns a delegate that is invoked when device was removed from this group.
 	 *
 	 * @return The delegate.
 	 */
-	virtual FOnLauncherDeviceGroupDeviceRemoved& OnDeviceRemoved( ) = 0;
+	virtual FOnLauncherDeviceGroupDeviceRemoved& OnDeviceRemoved() = 0;
 
 public:
 
 	/**
 	 * Virtual destructor.
 	 */
-	virtual ~ILauncherDeviceGroup( ) { }
+	virtual ~ILauncherDeviceGroup() { }
 };
diff --git a/Engine/Source/Developer/LauncherServices/Public/Interfaces/ILauncherProfile.h b/Engine/Source/Developer/LauncherServices/Public/Interfaces/ILauncherProfile.h
--- a/Engine/Source/Developer/LauncherServices/Public/Interfaces/ILauncherProfile.h
+++ b/Engine/Source/Developer/LauncherServices/Public/Interfaces/ILauncherProfile.h
@@ -154,13 +154,114 @@ namespace ELauncherProfileValidationErrors
 }
 
 
+/** Type definition for shared pointers to instances of ILauncherProfile. */
+typedef TSharedPtr<class ILauncherSimpleProfile> ILauncherSimpleProfilePtr;
+
+/** Type definition for shared references to instances of ILauncherProfile. */
+typedef TSharedRef<class ILauncherSimpleProfile> ILauncherSimpleProfileRef;
+
+/**
+* Interface for simple launcher profile.
+*/
+class ILauncherSimpleProfile
+{
+public:
+
+	/**
+	 * Gets the device name this profile is for.
+	 *
+	 * @return Device Name.
+	 */
+	virtual const FString& GetDeviceName() const = 0;
+
+	/**
+	 * Gets the device variant to use when deploying and launching.
+	 *
+	 * @return Device Variant name.
+	 *
+	 * @see SetDeviceVariant
+	 */
+	virtual FName GetDeviceVariant() const = 0;
+
+	/**
+	 * Gets the name of the build configuration.
+	 *
+	 * @return Build configuration name.
+	 *
+	 * @see SetBuildConfigurationName
+	 */
+	virtual EBuildConfigurations::Type GetBuildConfiguration() const = 0;
+
+	/**
+	 * Gets the selected cook mode.
+	 *
+	 * @return Cook mode.
+	 */
+	virtual ELauncherProfileCookModes::Type GetCookMode() const = 0;
+
+	/**
+	 * Updates the device name.
+	 *
+	 * @param InDeviceName - The new device name.
+	 */
+	virtual void SetDeviceName(const FString& InDeviceName) = 0;
+
+	/**
+	 * Sets the device variant.
+	 *
+	 * @param InVariant - The variant to set.
+	 *
+	 * @see GetDeviceVariant
+	 */
+	virtual void SetDeviceVariant(FName InVariant) = 0;
+
+	/**
+	 * Sets the build configuration.
+	 *
+	 * @param InConfiguration - The build configuration name to set.
+	 *
+	 * @see GetBuildConfigurationName
+	 */
+	virtual void SetBuildConfiguration(EBuildConfigurations::Type InConfiguration) = 0;
+
+	/**
+	 * Sets the cook mode.
+	 *
+	 * @param InMode - The cook mode.
+	 *
+	 * @see GetCookMode
+	 */
+	virtual void SetCookMode(ELauncherProfileCookModes::Type InMode) = 0;
+
+	/**
+	 * Serializes the simple profile from or into the specified archive.
+	 *
+	 * @param Archive - The archive to serialize from or into.
+	 *
+	 * @return true if the profile was serialized, false otherwise.
+	 */
+	virtual bool Serialize(FArchive& Archive) = 0;
+
+	/**
+	 * Sets all profile settings to their defaults.
+	 */
+	virtual void SetDefaults() = 0;
+
+public:
+
+	/**
+	* Virtual destructor.
+	*/
+	virtual ~ILauncherSimpleProfile() {}
+};
+
+
 /** Type definition for shared pointers to instances of ILauncherProfile. */
 typedef TSharedPtr<class ILauncherProfile> ILauncherProfilePtr;
 
 /** Type definition for shared references to instances of ILauncherProfile. */
 typedef TSharedRef<class ILauncherProfile> ILauncherProfileRef;
 
-
 /**
  * Delegate type for changing the device group to deploy to.
  *
@@ -171,7 +272,6 @@ DECLARE_MULTICAST_DELEGATE_OneParam(FOnLauncherProfileDeployedDeviceGroupChanged
 /** Delegate type for a change in project */
 DECLARE_MULTICAST_DELEGATE(FOnProfileProjectChanged);
 
-
 /**
  * Interface for launcher profile.
  */
@@ -193,6 +293,20 @@ public:
 	 */
 	virtual FString GetName( ) const = 0;
 
+	/**
+	 * Gets the human readable description of the profile.
+	 *
+	 * @return The profile description.
+	 */
+	virtual FString GetDescription() const = 0;
+
+	/**
+	 * Checks whether the last validation yielded any error.
+	 *
+	 * @return true if the any error is present, false otherwise.
+	 */
+	virtual bool HasValidationError() const = 0;
+
 	/**
 	 * Checks whether the last validation yielded the specified error.
 	 *
@@ -202,6 +316,11 @@ public:
 	 */
 	virtual bool HasValidationError( ELauncherProfileValidationErrors::Type Error ) const = 0;
 
+	/**
+	 * Gets the invalid platform, this is only valid when there is a platform centric validation error.
+	 *
+	 * @return string specifying the invalid platform.
+	 */
 	virtual FString GetInvalidPlatform() const = 0;
 
 	/**
@@ -244,6 +363,13 @@ public:
 	 */
 	virtual void SetName( const FString& NewName ) = 0;
 
+	/**
+	 * Updates the description of the profile.
+	 *
+	 * @param NewDescription - The new description of the profile.
+	 */
+	virtual void SetDescription(const FString& NewDescription) = 0;
+
 public:
 
 	/**
@@ -390,6 +516,14 @@ public:
 	 */
 	virtual FString GetPackageDirectory( ) const = 0;
 
+	/**
+	 * Checks whether the profile specifies a project.
+	 * Not specifying a project means that it can be used for any project.
+	 *
+	 * @return true if the profile specifies a project.
+	 */
+	virtual bool HasProjectSpecified() const = 0;
+
 	/**
 	 * Gets the name of the Unreal project to use.
 	 */
@@ -407,7 +541,7 @@ public:
 	 *
 	 * @see SetRocketProjactPath
 	 */
-	virtual const FString& GetProjectPath( ) const = 0;
+	virtual FString GetProjectPath( ) const = 0;
 
     /**
      * Gets the timeout time for the cook on the fly server
@@ -644,6 +778,15 @@ public:
 	 */
 	virtual void SetCookMode( ELauncherProfileCookModes::Type Mode ) = 0;
 
+	/**
+	 * Sets the cook options.
+	 *
+	 * @param Options - The cook options.
+	 *
+	 * @see GetCookOptions
+	 */
+	virtual void SetCookOptions(const FString& Options) = 0;
+
 	/**
 	 * Sets whether to pack with UnrealPak.
 	 *
@@ -725,6 +868,18 @@ public:
 	 */
 	virtual void SetPackageDirectory( const FString& Dir ) = 0;
 
+	/**
+	 * Sets whether this profile specifies the a project.
+	 *
+	 * @param Specified - Whether a project is specified.
+	 */
+	virtual void SetProjectSpecified(bool Specified) = 0;
+
+	/**
+	* Notifies the profile that the fallback project path changed.
+	*/
+	virtual void FallbackProjectUpdated() = 0;
+
 	/**
 	 * Sets the path to the Rocket project to use.
 	 *
diff --git a/Engine/Source/Developer/LauncherServices/Public/Interfaces/ILauncherProfileManager.h b/Engine/Source/Developer/LauncherServices/Public/Interfaces/ILauncherProfileManager.h
--- a/Engine/Source/Developer/LauncherServices/Public/Interfaces/ILauncherProfileManager.h
+++ b/Engine/Source/Developer/LauncherServices/Public/Interfaces/ILauncherProfileManager.h
@@ -60,6 +60,13 @@ public:
 	 */
 	virtual ILauncherDeviceGroupRef AddNewDeviceGroup( ) = 0;
 
+	/**
+	 * Creates a new device group but does not add it to the internal tracking.
+	 *
+	 * @return The new device group created.
+	 */
+	virtual ILauncherDeviceGroupRef CreateUnmanagedDeviceGroup() = 0;
+
 	/**
 	 * Gets the collection of device groups.
 	 *
@@ -83,6 +90,33 @@ public:
 	 */
 	virtual void RemoveDeviceGroup( const ILauncherDeviceGroupRef& DeviceGroup ) = 0;
 
+public:
+
+	/**
+	 * Finds or Adds then returns a simple profile for the specified Device
+	 *
+	 * @param DeviceName - The name of the device we want the simple profile for.
+	 *
+	 * @return The simple profile for the specified device.
+	 */
+	virtual ILauncherSimpleProfilePtr FindOrAddSimpleProfile(const FString& DeviceName) = 0;
+
+	/**
+	 * Gets the simple profile for the specified device.
+	 *
+	 * @param DeviceName - The name of the device we want the simple profile for.
+	 *
+	 * @return The simple profile for the specified device, or NULL if the simple profile doesn't exist.
+	 */
+	virtual ILauncherSimpleProfilePtr FindSimpleProfile(const FString& DeviceName) = 0;
+
+	/**
+	 * Deletes the given simple profile.
+	 *
+	 * @param Profile - The simple profile to delete.
+	 */
+	virtual void RemoveSimpleProfile(const ILauncherSimpleProfileRef& SimpleProfile) = 0;
+
 public:
 
 	/**
@@ -92,6 +126,15 @@ public:
 	 */
 	virtual ILauncherProfileRef AddNewProfile( ) = 0;
 
+	/**
+	 * Creates a new profile but does not add it to the internal tracking.
+	 *
+	 * @param ProfileName - The name of the profile to create.
+	 *
+	 * @return The new profile created.
+	 */
+	virtual ILauncherProfileRef CreateUnsavedProfile(FString ProfileName) = 0;
+
 	/**
 	 * Adds the given profile to the list of managed profiles.
 	 *
@@ -163,7 +206,7 @@ public:
 	 *
 	 * @see LoadProfile
 	 */
-	virtual void SaveProfile( const ILauncherProfileRef& Profile, FArchive& Archive ) = 0;
+	virtual void SaveProfile(const ILauncherProfileRef& Profile) = 0;
 
 public:
 
@@ -185,6 +228,36 @@ public:
 	 */
 	virtual void SaveSettings( ) = 0;
 
+public:
+
+	/**
+	 * Gets the name of the Unreal project to use.
+	 */
+	virtual FString GetProjectName() const = 0;
+
+	/**
+	 * Gets the base project path for the project
+	 */
+	virtual FString GetProjectBasePath() const = 0;
+
+	/**
+	 * Gets the full path to the Unreal project to use.
+	 *
+	 * @return The path.
+	 *
+	 * @see SetProjectPath
+	 */
+	virtual FString GetProjectPath() const = 0;
+
+	/**
+	 * Sets the path to the Unreal project to use.
+	 *
+	 * @param Path - The full path to the project.
+	 *
+	 * @see GetProjectPath
+	 */
+	virtual void SetProjectPath(const FString& InProjectPath) = 0;
+
 public:
 
 	/**
diff --git a/Engine/Source/Developer/Linux/LinuxTargetPlatform/Private/LinuxTargetPlatform.h b/Engine/Source/Developer/Linux/LinuxTargetPlatform/Private/LinuxTargetPlatform.h
--- a/Engine/Source/Developer/Linux/LinuxTargetPlatform/Private/LinuxTargetPlatform.h
+++ b/Engine/Source/Developer/Linux/LinuxTargetPlatform/Private/LinuxTargetPlatform.h
@@ -10,6 +10,8 @@
 #include "StaticMeshResources.h"
 #endif // WITH_ENGINE
 
+#define LOCTEXT_NAMESPACE "TLinuxTargetPlatform"
+
 /**
  * Template for Linux target platforms
  */
@@ -19,7 +21,8 @@ class TLinuxTargetPlatform
 {
 public:
 
-	typedef TTargetPlatformBase<FLinuxPlatformProperties<HAS_EDITOR_DATA, IS_DEDICATED_SERVER, IS_CLIENT_ONLY> > TSuper;
+	typedef FLinuxPlatformProperties<HAS_EDITOR_DATA, IS_DEDICATED_SERVER, IS_CLIENT_ONLY> TProperties;
+	typedef TTargetPlatformBase<TProperties> TSuper;
 
 	/**
 	 * Default constructor.
@@ -62,6 +65,17 @@ public:
 
 	virtual void EnableDeviceCheck(bool OnOff) override {}
 
+	virtual bool AddDevice(const FString& DeviceName, bool bDefault) override
+	{
+		FTargetDeviceId UATFriendlyId(TEXT("Linux"), DeviceName);
+		FLinuxTargetDevicePtr Device = MakeShareable(new FLinuxTargetDevice(*this, UATFriendlyId, DeviceName));
+		if (Device.IsValid())
+		{
+			DeviceDiscoveredEvent.Broadcast(Device.ToSharedRef());
+		}
+		return true;
+	}
+
 	virtual void GetAllDevices( TArray<ITargetDevicePtr>& OutDevices ) const override
 	{
 		// TODO: ping all the machines in a local segment and/or try to connect to port 22 of those that respond
@@ -98,9 +112,7 @@ public:
 		{
 			return LocalDevice;
 		}
-
-		FTargetDeviceId UATFriendlyId(TEXT("Linux"), DeviceId.GetDeviceName());
-		return MakeShareable(new FLinuxTargetDevice(*this, UATFriendlyId, DeviceId.GetDeviceName()));
+		return nullptr;
 	}
 
 	virtual bool IsRunningPlatform( ) const override
@@ -174,6 +186,41 @@ public:
 	}
 #endif //WITH_ENGINE
 
+	virtual bool SupportsVariants() const override
+	{
+		return true;
+	}
+
+	virtual FText GetVariantDisplayName() const override
+	{
+		if (IS_DEDICATED_SERVER)
+		{
+			return LOCTEXT("LinuxServerVariantTitle", "Dedicated Server");
+		}
+
+		if (HAS_EDITOR_DATA)
+		{
+			return LOCTEXT("LinuxClientEditorDataVariantTitle", "Client with Editor Data");
+		}
+
+		if (IS_CLIENT_ONLY)
+		{
+			return LOCTEXT("LinuxClientOnlyVariantTitle", "Client only");
+		}
+
+		return LOCTEXT("LinuxClientVariantTitle", "Client");
+	}
+
+	virtual FText GetVariantTitle() const override
+	{
+		return LOCTEXT("LinuxVariantTitle", "Build Type");
+	}
+
+	virtual float GetVariantPriority() const override
+	{
+		return TProperties::GetVariantPriority();
+	}
+
 	DECLARE_DERIVED_EVENT(FAndroidTargetPlatform, ITargetPlatform::FOnTargetDeviceDiscovered, FOnTargetDeviceDiscovered);
 	virtual FOnTargetDeviceDiscovered& OnDeviceDiscovered( ) override
 	{
@@ -215,3 +262,5 @@ private:
 	// Holds an event delegate that is executed when a target device has been lost, i.e. disconnected or timed out.
 	FOnTargetDeviceLost DeviceLostEvent;
 };
+
+#undef LOCTEXT_NAMESPACE
diff --git a/Engine/Source/Developer/Mac/MacTargetPlatform/Private/GenericMacTargetPlatform.h b/Engine/Source/Developer/Mac/MacTargetPlatform/Private/GenericMacTargetPlatform.h
--- a/Engine/Source/Developer/Mac/MacTargetPlatform/Private/GenericMacTargetPlatform.h
+++ b/Engine/Source/Developer/Mac/MacTargetPlatform/Private/GenericMacTargetPlatform.h
@@ -10,6 +10,8 @@
 #include "StaticMeshResources.h"
 #endif // WITH_ENGINE
 
+#define LOCTEXT_NAMESPACE "TGenericMacTargetPlatform"
+
 /**
  * Template for Mac target platforms
  */
@@ -19,7 +21,8 @@ class TGenericMacTargetPlatform
 {
 public:
 
-	typedef TTargetPlatformBase<FMacPlatformProperties<HAS_EDITOR_DATA, IS_DEDICATED_SERVER, IS_CLIENT_ONLY> > TSuper;
+	typedef FMacPlatformProperties<HAS_EDITOR_DATA, IS_DEDICATED_SERVER, IS_CLIENT_ONLY> TProperties;
+	typedef TTargetPlatformBase<TProperties> TSuper;
 
 	/**
 	 * Default constructor.
@@ -141,6 +144,46 @@ return TSuper::SupportsFeature(Feature);
 	}
 #endif //WITH_ENGINE
 
+
+	virtual bool SupportsVariants() const override
+	{
+		return true;
+	}
+
+
+	virtual FText GetVariantDisplayName() const override
+	{
+		if (IS_DEDICATED_SERVER)
+		{
+			return LOCTEXT("MacServerVariantTitle", "Dedicated Server");
+		}
+
+		if (HAS_EDITOR_DATA)
+		{
+			return LOCTEXT("MacClientEditorDataVariantTitle", "Client with Editor Data");
+		}
+
+		if (IS_CLIENT_ONLY)
+		{
+			return LOCTEXT("MacClientOnlyVariantTitle", "Client only");
+		}
+
+		return LOCTEXT("MacClientVariantTitle", "Client");
+	}
+
+
+	virtual FText GetVariantTitle() const override
+	{
+		return LOCTEXT("MacVariantTitle", "Build Type");
+	}
+
+
+	virtual float GetVariantPriority() const override
+	{
+		return TProperties::GetVariantPriority();
+	}
+
+
 	DECLARE_DERIVED_EVENT(TGenericMacTargetPlatform, ITargetPlatform::FOnTargetDeviceDiscovered, FOnTargetDeviceDiscovered);
 	virtual FOnTargetDeviceDiscovered& OnDeviceDiscovered( ) override
 	{
@@ -179,3 +222,5 @@ private:
 	// Holds an event delegate that is executed when a target device has been lost, i.e. disconnected or timed out.
 	FOnTargetDeviceLost DeviceLostEvent;
 };
+
+#undef LOCTEXT_NAMESPACE
diff --git a/Engine/Source/Developer/ProjectLauncher/Private/Models/ProjectLauncherCommands.h b/Engine/Source/Developer/ProjectLauncher/Private/Models/ProjectLauncherCommands.h
new file mode 100644
--- /dev/null
+++ b/Engine/Source/Developer/ProjectLauncher/Private/Models/ProjectLauncherCommands.h
@@ -0,0 +1,60 @@
+// Copyright 1998-2014 Epic Games, Inc. All Rights Reserved.
+
+#pragma once
+
+
+/**
+ * The project launcher UI commands
+ */
+class FProjectLauncherCommands : public TCommands<FProjectLauncherCommands>
+{
+public:
+
+	/**
+	 * Default constructor.
+	 */
+	FProjectLauncherCommands( )
+		: TCommands<FProjectLauncherCommands>(
+			"ProjectLauncherCommand",
+			NSLOCTEXT("Contexts", "ProjectLauncherCommand", "Project Launcher Command"),
+			NAME_None,
+			FEditorStyle::GetStyleSetName()
+	)
+	{ }
+	
+public:
+
+	// TCommands interface
+
+	PRAGMA_DISABLE_OPTIMIZATION
+	virtual void RegisterCommands( ) override
+	{
+		UI_COMMAND(QuickLaunch, "Quick Launch", "Builds, cooks, and launches a build.", EUserInterfaceActionType::ToggleButton, FInputGesture(EModifierKey::Control, EKeys::L));
+		UI_COMMAND(CreateBuild, "Build", "Creates a build.", EUserInterfaceActionType::ToggleButton, FInputGesture(EModifierKey::Control, EKeys::B));
+		UI_COMMAND(DeployBuild, "Deploy Build", "Deploys a pre-made build.", EUserInterfaceActionType::ToggleButton, FInputGesture(EModifierKey::Control, EKeys::D));
+		UI_COMMAND(AdvancedBuild, "Advanced...", "Advanced launcher.", EUserInterfaceActionType::ToggleButton, FInputGesture(EModifierKey::Control, EKeys::A));
+		UI_COMMAND(CloseSettings, "Back", "Close the launch profile settings editor and go back to the launcher.", EUserInterfaceActionType::Button, FInputGesture(EModifierKey::None, EKeys::Escape));
+		UI_COMMAND(RenameProfile, "Rename", "Rename the launch profile.", EUserInterfaceActionType::Button, FInputGesture());// , FInputGesture(EModifierKey::None, EKeys::F2));
+		UI_COMMAND(DuplicateProfile, "Duplicate", "Duplicate the launch profile.", EUserInterfaceActionType::Button, FInputGesture());//, FInputGesture(EModifierKey::Control, EKeys::W));
+		UI_COMMAND(DeleteProfile, "Delete", "Delete the launch profile.", EUserInterfaceActionType::Button, FInputGesture());//, FInputGesture(EModifierKey::Control, EKeys::Delete));
+	}
+	PRAGMA_ENABLE_OPTIMIZATION
+
+public:
+
+	TSharedPtr<FUICommandInfo> CreateBuild;
+
+	TSharedPtr<FUICommandInfo> DeployBuild;
+
+	TSharedPtr<FUICommandInfo> QuickLaunch;
+
+	TSharedPtr<FUICommandInfo> AdvancedBuild;
+
+	TSharedPtr<FUICommandInfo> CloseSettings;
+
+	TSharedPtr<FUICommandInfo> RenameProfile;
+
+	TSharedPtr<FUICommandInfo> DuplicateProfile;
+
+	TSharedPtr<FUICommandInfo> DeleteProfile;
+};
diff --git a/Engine/Source/Developer/ProjectLauncher/Private/Models/ProjectLauncherModel.h b/Engine/Source/Developer/ProjectLauncher/Private/Models/ProjectLauncherModel.h
new file mode 100644
--- /dev/null
+++ b/Engine/Source/Developer/ProjectLauncher/Private/Models/ProjectLauncherModel.h
@@ -0,0 +1,251 @@
+// Copyright 1998-2014 Epic Games, Inc. All Rights Reserved.
+
+#pragma once
+
+
+namespace ELauncherPanels
+{
+	/**
+	 * Enumerates available launcher panels.
+	 */
+	enum Type
+	{
+		NoTask = 0,
+		Launch,
+		ProfileEditor,
+		Progress,
+		End,
+	};
+}
+
+
+/** Type definition for shared pointers to instances of FProjectLauncherModel. */
+typedef TSharedPtr<class FProjectLauncherModel> FProjectLauncherModelPtr;
+
+/** Type definition for shared references to instances of FProjectLauncherModel. */
+typedef TSharedRef<class FProjectLauncherModel> FProjectLauncherModelRef;
+
+
+/**
+ * Delegate type for launcher profile selection changes.
+ *
+ * The first parameter is the newly selected profile (or NULL if none was selected).
+ * The second parameter is the old selected profile (or NULL if none was previous selected).
+ */
+DECLARE_MULTICAST_DELEGATE_TwoParams(FOnSelectedSProjectLauncherProfileChanged, const ILauncherProfilePtr&, const ILauncherProfilePtr&)
+
+
+/**
+ * Implements the data model for the session launcher.
+ */
+class FProjectLauncherModel
+{
+public:
+
+	/**
+	 * Creates and initializes a new instance.
+	 *
+	 * @param InDeviceProxyManager - The device proxy manager to use.
+	 * @param InLauncher - The launcher to use.
+	 * @param InProfileManager - The profile manager to use.
+	 */
+	FProjectLauncherModel(const ITargetDeviceProxyManagerRef& InDeviceProxyManager, const ILauncherRef& InLauncher, const ILauncherProfileManagerRef& InProfileManager)
+		: DeviceProxyManager(InDeviceProxyManager)
+		, SProjectLauncher(InLauncher)
+		, ProfileManager(InProfileManager)
+	{
+		ProfileManager->OnProfileAdded().AddRaw(this, &FProjectLauncherModel::HandleProfileManagerProfileAdded);
+		ProfileManager->OnProfileRemoved().AddRaw(this, &FProjectLauncherModel::HandleProfileManagerProfileRemoved);
+
+		LoadConfig();
+	}
+
+	/**
+	 * Destructor.
+	 */
+	~FProjectLauncherModel()
+	{
+		ProfileManager->OnProfileAdded().RemoveAll(this);
+		ProfileManager->OnProfileRemoved().RemoveAll(this);
+
+		SaveConfig();
+	}
+
+public:
+
+	/**
+	 * Gets the model's device proxy manager.
+	 *
+	 * @return Device proxy manager.
+	 */
+	const ITargetDeviceProxyManagerRef& GetDeviceProxyManager() const
+	{
+		return DeviceProxyManager;
+	}
+
+	/**
+	 * Gets the model's launcher.
+	 *
+	 * @return SProjectLauncher.
+	 */
+	const ILauncherRef& GetSProjectLauncher() const
+	{
+		return SProjectLauncher;
+	}
+
+	/**
+	 * Get the model's profile manager.
+	 *
+	 * @return Profile manager.
+	 */
+	const ILauncherProfileManagerRef& GetProfileManager() const
+	{
+		return ProfileManager;
+	}
+
+
+	/**
+	 * Gets the active profile.
+	 *
+	 * @return The active profile.
+	 *
+	 * @see GetProfiles
+	 */
+	const ILauncherProfilePtr& GetSelectedProfile() const
+	{
+		return SelectedProfile;
+	}
+
+	/**
+	 * Sets the active profile.
+	 *
+	 * @param InProfile The profile to assign as active,, or NULL to deselect all.
+	 */
+	void SelectProfile(const ILauncherProfilePtr& Profile)
+	{
+		if (!Profile.IsValid() || ProfileManager->GetAllProfiles().Contains(Profile))
+		{
+			if (Profile != SelectedProfile)
+			{
+				ILauncherProfilePtr& PreviousProfile = SelectedProfile;
+				SelectedProfile = Profile;
+
+				ProfileSelectedDelegate.Broadcast(Profile, PreviousProfile);
+			}
+		}
+	}
+
+public:
+
+	/**
+	 * Returns a delegate to be invoked when the profile list has been modified.
+	 *
+	 * @return The delegate.
+	 */
+	FSimpleMulticastDelegate& OnProfileListChanged()
+	{
+		return ProfileListChangedDelegate;
+	}
+
+	/**
+	 * Returns a delegate to be invoked when the selected profile changed.
+	 *
+	 * @return The delegate.
+	 */
+	FOnSelectedSProjectLauncherProfileChanged& OnProfileSelected()
+	{
+		return ProfileSelectedDelegate;
+	}
+
+protected:
+
+	/*
+	 * Load all profiles from disk.
+	 */
+	void LoadConfig()
+	{
+		// restore the previous project selection
+		FString ProjectPath;
+		if (FPaths::IsProjectFilePathSet())
+		{
+			ProjectPath = FPaths::GetProjectFilePath();
+		}
+		else if (FGameProjectHelper::IsGameAvailable(GGameName))
+		{
+			ProjectPath = FPaths::RootDir() / GGameName / GGameName + TEXT(".uproject");
+		}
+		else if (GConfig != NULL)
+		{
+			GConfig->GetString(TEXT("FProjectLauncherModel"), TEXT("SelectedProjectPath"), ProjectPath, GEngineIni);
+		}
+		ProfileManager->SetProjectPath(ProjectPath);
+	}
+
+	/*
+	 * Saves all profiles to disk.
+	 */
+	void SaveConfig()
+	{
+		if (GConfig != NULL && !FPaths::IsProjectFilePathSet() && !FGameProjectHelper::IsGameAvailable(GGameName))
+		{
+			FString ProjectPath = ProfileManager->GetProjectPath();
+			GConfig->SetString(TEXT("FProjectLauncherModel"), TEXT("SelectedProjectPath"), *ProjectPath, GEngineIni);
+		}
+	}
+
+private:
+
+	// Callback for when a profile was added to the profile manager.
+	void HandleProfileManagerProfileAdded(const ILauncherProfileRef& Profile)
+	{
+		ProfileListChangedDelegate.Broadcast();
+
+		SelectProfile(Profile);
+	}
+
+	// Callback for when a profile was removed from the profile manager.
+	void HandleProfileManagerProfileRemoved(const ILauncherProfileRef& Profile)
+	{
+		ProfileListChangedDelegate.Broadcast();
+
+		if (Profile == SelectedProfile)
+		{
+			const TArray<ILauncherProfilePtr>& Profiles = ProfileManager->GetAllProfiles();
+
+			if (Profiles.Num() > 0)
+			{
+				SelectProfile(Profiles[0]);
+			}
+			else
+			{
+				SelectProfile(NULL);
+			}
+		}
+	}
+
+private:
+
+	// Holds a pointer to the device proxy manager.
+	ITargetDeviceProxyManagerRef DeviceProxyManager;
+
+	// Holds a pointer to the launcher.
+	ILauncherRef SProjectLauncher;
+
+	// Holds a pointer to the profile manager.
+	ILauncherProfileManagerRef ProfileManager;
+
+	// Holds a pointer to active profile.
+	ILauncherProfilePtr SelectedProfile;
+
+private:
+
+	// Holds a delegate to be invoked when the project path has been modified.
+	FSimpleMulticastDelegate ProjectPathChangedDelegate;
+
+	// Holds a delegate to be invoked when the profile list has been modified.
+	FSimpleMulticastDelegate ProfileListChangedDelegate;
+
+	// Holds a delegate to be invoked when the selected profile changed.
+	FOnSelectedSProjectLauncherProfileChanged ProfileSelectedDelegate;
+
+};
diff --git a/Engine/Source/Developer/ProjectLauncher/Private/ProjectLauncherModule.cpp b/Engine/Source/Developer/ProjectLauncher/Private/ProjectLauncherModule.cpp
new file mode 100644
--- /dev/null
+++ b/Engine/Source/Developer/ProjectLauncher/Private/ProjectLauncherModule.cpp
@@ -0,0 +1,75 @@
+// Copyright 1998-2014 Epic Games, Inc. All Rights Reserved.
+
+#include "ProjectLauncherPrivatePCH.h"
+
+
+static const FName ProjectLauncherTabName("ProjectLauncher");
+
+
+/**
+ * Implements the SessionSProjectLauncher module.
+ */
+class FProjectLauncherModule
+	: public IProjectLauncherModule
+{
+public:
+
+	// IProjectLauncherModule interface
+
+	virtual TSharedRef<class SWidget> CreateSProjectLauncherProgressPanel( const ILauncherWorkerRef& LauncherWorker ) override
+	{
+		TSharedRef<SProjectLauncherProgress> Panel = SNew(SProjectLauncherProgress);
+
+		Panel->SetLauncherWorker(LauncherWorker);
+
+		return Panel;
+	}
+
+public:
+
+	// IModuleInterface interface
+	
+	virtual void StartupModule( ) override
+	{
+		FGlobalTabmanager::Get()->RegisterTabSpawner(ProjectLauncherTabName, FOnSpawnTab::CreateRaw(this, &FProjectLauncherModule::SpawnProjectLauncherTab))
+			.SetDisplayName(NSLOCTEXT("FProjectLauncherModule", "ProjectLauncherTabTitle", "Project Launcher"))
+			.SetTooltipText(NSLOCTEXT("FProjectLauncherModule", "ProjectLauncherTooltipText", "Open the Project Launcher tab."))
+			.SetIcon(FSlateIcon(FEditorStyle::GetStyleSetName(), "Launcher.TabIcon"));
+	}
+
+	virtual void ShutdownModule( ) override
+	{
+		FGlobalTabmanager::Get()->UnregisterTabSpawner(ProjectLauncherTabName);
+	}
+
+private:
+
+	/**
+	 * Creates a new session launcher tab.
+	 *
+	 * @param SpawnTabArgs The arguments for the tab to spawn.
+	 * @return The spawned tab.
+	 */
+	TSharedRef<SDockTab> SpawnProjectLauncherTab(const FSpawnTabArgs& SpawnTabArgs)
+	{
+		const TSharedRef<SDockTab> DockTab = SNew(SDockTab)
+			.Icon(FEditorStyle::GetBrush("Launcher.TabIcon"))
+			.TabRole(ETabRole::NomadTab);
+
+		ILauncherServicesModule& ProjectLauncherServicesModule = FModuleManager::LoadModuleChecked<ILauncherServicesModule>("LauncherServices");
+		ITargetDeviceServicesModule& TargetDeviceServicesModule = FModuleManager::LoadModuleChecked<ITargetDeviceServicesModule>("TargetDeviceServices");
+
+		FProjectLauncherModelRef Model = MakeShareable(new FProjectLauncherModel(
+			TargetDeviceServicesModule.GetDeviceProxyManager(),
+			ProjectLauncherServicesModule.CreateLauncher(),
+			ProjectLauncherServicesModule.GetProfileManager()
+		));
+
+		DockTab->SetContent(SNew(SProjectLauncher, DockTab, SpawnTabArgs.GetOwnerWindow(), Model));
+
+		return DockTab;
+	}
+};
+
+
+IMPLEMENT_MODULE(FProjectLauncherModule, SProjectLauncher);
