Quick fix for Launch On

diff --git a/Engine/Source/Developer/LauncherServices/Private/Profiles/LauncherProfile.h b/Engine/Source/Developer/LauncherServices/Private/Profiles/LauncherProfile.h
--- a/Engine/Source/Developer/LauncherServices/Private/Profiles/LauncherProfile.h
+++ b/Engine/Source/Developer/LauncherServices/Private/Profiles/LauncherProfile.h
@@ -563,6 +563,9 @@ public:
 		{
 			ProjectPath = FString();
 		}
+
+		// Use the locally specified project path is resolving through the root isn't working
+		ProjectSpecified = GetProjectPath().IsEmpty();
 		
 		BuildConfiguration = FApp::GetBuildConfiguration();
 		FInternationalization& I18N = FInternationalization::Get();
