converted LauncherWorker to utilize a single UAT command
fixed problem with IP not being set for CotF
#ue4

diff --git a/Engine/Source/Developer/LauncherServices/LauncherServices.Build.cs b/Engine/Source/Developer/LauncherServices/LauncherServices.Build.cs
--- a/Engine/Source/Developer/LauncherServices/LauncherServices.Build.cs
+++ b/Engine/Source/Developer/LauncherServices/LauncherServices.Build.cs
@@ -21,6 +21,7 @@ namespace UnrealBuildTool.Rules
 			PrivateDependencyModuleNames.AddRange(
 				new string[] {
 					"CoreUObject",
+					"DesktopPlatform",
 					"SessionMessages",
 					"TargetPlatform",
 					"UnrealEdMessages",
diff --git a/Engine/Source/Developer/LauncherServices/Private/Launcher/LauncherBuildCommands.h b/Engine/Source/Developer/LauncherServices/Private/Launcher/LauncherBuildCommands.h
deleted file mode 100644
diff --git a/Engine/Source/Developer/LauncherServices/Private/Launcher/LauncherCookCommands.h b/Engine/Source/Developer/LauncherServices/Private/Launcher/LauncherCookCommands.h
deleted file mode 100644
diff --git a/Engine/Source/Developer/LauncherServices/Private/Launcher/LauncherDeployCommands.h b/Engine/Source/Developer/LauncherServices/Private/Launcher/LauncherDeployCommands.h
deleted file mode 100644
diff --git a/Engine/Source/Developer/LauncherServices/Private/Launcher/LauncherLaunchCommands.h b/Engine/Source/Developer/LauncherServices/Private/Launcher/LauncherLaunchCommands.h
deleted file mode 100644
diff --git a/Engine/Source/Developer/LauncherServices/Private/Launcher/LauncherPackageCommands.h b/Engine/Source/Developer/LauncherServices/Private/Launcher/LauncherPackageCommands.h
deleted file mode 100644
diff --git a/Engine/Source/Developer/LauncherServices/Private/Launcher/LauncherUATTask.h b/Engine/Source/Developer/LauncherServices/Private/Launcher/LauncherUATTask.h
--- a/Engine/Source/Developer/LauncherServices/Private/Launcher/LauncherUATTask.h
+++ b/Engine/Source/Developer/LauncherServices/Private/Launcher/LauncherUATTask.h
@@ -15,20 +15,16 @@ class FLauncherUATTask
 {
 public:
 
-	/**
-	 * Creates and initializes a new instance.
-	 *
-	 * @param InCommand - the command used for the task
-	 * @param InTargetPlatform - the target platform for the task
-	 * @param InName - The name of the task.
-	 */
-	FLauncherUATTask( const TSharedPtr<FLauncherUATCommand>& InCommand, const ITargetPlatform& InTargetPlatform, const FString& InName, void* InReadPipe, void* InWritePipe, const FString& InEditorExe )
-		: FLauncherTask(InName, InCommand->GetDesc(), InReadPipe, InWritePipe)
-		, TaskCommand(InCommand)
-		, TargetPlatform(InTargetPlatform)
+	FLauncherUATTask( const FString& InCommandLine, const FString& InName, const FString& InDesc, void* InReadPipe, void* InWritePipe, const FString& InEditorExe, FProcHandle& InProcessHandle, ILauncherWorker* InWorker, const FString& InCommandEnd)
+		: FLauncherTask(InName, InDesc, InReadPipe, InWritePipe)
+		, CommandLine(InCommandLine)
 		, EditorExe(InEditorExe)
+		, ProcessHandle(InProcessHandle)
+		, CommandText(InCommandEnd)
 	{
 		NoCompile = TEXT(" -nocompile");
+		EndTextFound = false;
+		InWorker->OnOutputReceived().AddRaw(this, &FLauncherUATTask::HandleOutputReceived);
 	}
 
 	static bool FirstTimeCompile;
@@ -60,31 +56,6 @@ protected:
 		FString Executable = TEXT("RunUAT.bat");
 #endif
 
-		const PlatformInfo::FPlatformInfo& PlatformInfo = TargetPlatform.GetPlatformInfo();
-
-		// switch server and no editor platforms to the proper type
-		FName PlatformName = PlatformInfo.TargetPlatformName;
-		if (PlatformName == FName("LinuxServer") || PlatformName == FName("LinuxNoEditor"))
-		{
-			PlatformName = FName("Linux");
-		}
-		else if (PlatformName == FName("WindowsServer") || PlatformName == FName("WindowsNoEditor") || PlatformName == FName("Windows"))
-		{
-			PlatformName = FName("Win64");
-		}
-        else if (PlatformName == FName("MacNoEditor"))
-        {
-            PlatformName = FName("Mac");
-        }
-
-		// Append any extra UAT flags specified for this platform flavor
-		FString OptionalParams;
-		if (!PlatformInfo.UATCommandLine.IsEmpty())
-		{
-			OptionalParams += TEXT(" ");
-			OptionalParams += PlatformInfo.UATCommandLine;
-		}
-
         // check for rocket
         FString Rocket;
         if ( FRocketSupport::IsRocket() )
@@ -93,39 +64,32 @@ protected:
         }
         
 		// base UAT command arguments
-		FString CommandLine;
+		FString UATCommandLine;
 		FString ProjectPath = *ChainState.Profile->GetProjectPath();
 		ProjectPath = FPaths::ConvertRelativePathToFull(ProjectPath);
-		CommandLine = FString::Printf(TEXT("BuildCookRun -project=\"%s\" -noP4 -platform=%s -clientconfig=%s -serverconfig=%s"),
+		UATCommandLine = FString::Printf(TEXT("BuildCookRun -project=\"%s\" -noP4 -clientconfig=%s -serverconfig=%s"),
 			*ProjectPath,
-			*PlatformName.ToString(),
 			*ConfigStrings[ChainState.Profile->GetBuildConfiguration()],
 			*ConfigStrings[ChainState.Profile->GetBuildConfiguration()]);
 
-		CommandLine += NoCompile;
-        CommandLine += Rocket;
-		CommandLine += OptionalParams;
+		UATCommandLine += NoCompile;
+        UATCommandLine += Rocket;
 
 		// specify the path to the editor exe if necessary
 		if(EditorExe.Len() > 0)
 		{
-			CommandLine += FString::Printf(TEXT(" -ue4exe=\"%s\""), *EditorExe);
+			UATCommandLine += FString::Printf(TEXT(" -ue4exe=\"%s\""), *EditorExe);
 		}
 
 		// specialized command arguments for this particular task
-		CommandLine += TaskCommand->GetArguments(ChainState);
-
-		// pre-execute any command
-		if (!TaskCommand->PreExecute(ChainState))
-		{
-			return false;
-		}
+		UATCommandLine += CommandLine;
 
 		// launch UAT and monitor its progress
-		FProcHandle ProcessHandle = FPlatformProcess::CreateProc(*(ExecutablePath / Executable), *CommandLine, false, true, true, NULL, 0, *ExecutablePath, WritePipe);
+		ProcessHandle = FPlatformProcess::CreateProc(*(ExecutablePath / Executable), *UATCommandLine, false, true, true, NULL, 0, *ExecutablePath, WritePipe);
 
-		while (!TaskCommand->IsComplete() || FPlatformProcess::IsProcRunning(ProcessHandle))
+		while (FPlatformProcess::IsProcRunning(ProcessHandle) && !EndTextFound)
 		{
+
 			if (GetStatus() == ELauncherTaskStatus::Canceling)
 			{
 				FPlatformProcess::TerminateProc(ProcessHandle, true);
@@ -135,12 +99,7 @@ protected:
 			FPlatformProcess::Sleep(0.25);
 		}
 
-		if (!FPlatformProcess::GetProcReturnCode(ProcessHandle, &Result))
-		{
-			return false;
-		}
-
-		if (!TaskCommand->PostExecute(ChainState))
+		if (!EndTextFound && !FPlatformProcess::GetProcReturnCode(ProcessHandle, &Result))
 		{
 			return false;
 		}
@@ -148,17 +107,28 @@ protected:
 		return (Result == 0);
 	}
 
-private:
+	void HandleOutputReceived(const FString& InMessage)
+	{
+		EndTextFound |= InMessage.Contains(CommandText);
+	}
 
-	// command for this task
-	const TSharedPtr<FLauncherUATCommand> TaskCommand;
+private:
 
-	// Holds a pointer to the target platform.
-	const ITargetPlatform& TargetPlatform;
+	// Command line
+	FString CommandLine;
 
 	// Holds the no compile flag
 	FString NoCompile;
 
 	// The editor executable that UAT should use
 	FString EditorExe;
+
+	// process handle
+	FProcHandle& ProcessHandle;
+
+	// The editor executable that UAT should use
+	FString CommandText;
+
+	bool EndTextFound;
 };
+
diff --git a/Engine/Source/Developer/LauncherServices/Private/Launcher/LauncherWorker.cpp b/Engine/Source/Developer/LauncherServices/Private/Launcher/LauncherWorker.cpp
--- a/Engine/Source/Developer/LauncherServices/Private/Launcher/LauncherWorker.cpp
+++ b/Engine/Source/Developer/LauncherServices/Private/Launcher/LauncherWorker.cpp
@@ -1,7 +1,9 @@
 // Copyright 1998-2014 Epic Games, Inc. All Rights Reserved.
 
 #include "LauncherServicesPrivatePCH.h"
+#include "PlatformInfo.h"
 
+#define LOCTEXT_NAMESPACE "LauncherWorker"
 
 bool FLauncherUATTask::FirstTimeCompile = true;
 
@@ -177,37 +179,99 @@ void FLauncherWorker::OnTaskCompleted(const FString& TaskName)
 }
 
 
-/* FLauncherWorker implementation
- *****************************************************************************/
-
-void FLauncherWorker::CreateAndExecuteTasks( const ILauncherProfileRef& InProfile )
+FString FLauncherWorker::CreateUATCommand( const ILauncherProfileRef& InProfile, const TArray<FString>& InPlatforms, TArray<FCommandDesc>& OutCommands, FString& CommandStart )
 {
-	FPlatformProcess::CreatePipe(ReadPipe, WritePipe);
+	CommandStart = TEXT("");
+	FString UATCommand = TEXT("");
+	FGuid InstanceId(FGuid::NewGuid());
+	FGuid SessionId(FGuid::NewGuid());
+	FString InitialMap = InProfile->GetDefaultLaunchRole()->GetInitialMap();
+	if (InitialMap.IsEmpty() && InProfile->GetCookedMaps().Num() == 1)
+	{
+		InitialMap = InProfile->GetCookedMaps()[0];
+	}
 
-	// create task chains
-	TaskChain = MakeShareable(new FLauncherVerifyProfileTask());
-	TArray<FString> Platforms;
-	if (InProfile->GetCookMode() == ELauncherProfileCookModes::ByTheBook || InProfile->IsBuilding())
+	// game command line
+	FString CommandLine = FString::Printf(TEXT(" -cmdline=\"%s -Messaging\""),
+		*InitialMap);
+
+	FString AdditionalCommandLine = FString::Printf(TEXT(" -addcmdline=\"%s -SessionId=%s -SessionOwner=%s -SessionName='%s'%s\""),
+		*InitialMap,
+		*InstanceId.ToString(),
+		*SessionId.ToString(),
+		FPlatformProcess::UserName(false),
+		*InProfile->GetName(),
+		InProfile->GetLaunchRoles().Num() > 0 ? (InProfile->GetLaunchRoles()[0]->IsVsyncEnabled() ? TEXT(" -vsync") : TEXT("")) : TEXT(""));
+
+	// staging directory
+	FString StageDirectory = TEXT("");
+	auto PackageDirectory = InProfile->GetPackageDirectory();
+	if (PackageDirectory.Len() > 0)
 	{
-		Platforms = InProfile->GetCookedPlatforms();
+		StageDirectory += FString::Printf(TEXT(" -stagingdirectory=\"%s\""), *PackageDirectory);
 	}
-	TSharedPtr<FLauncherTask> PerPlatformTask = TaskChain;
-	TSharedPtr<FLauncherTask> PerPlatformBuildTask = NULL;
-	TSharedPtr<FLauncherTask> PerPlatformCookTask = NULL;
-	TSharedPtr<FLauncherTask> PerPlatformPackageTask = NULL;
-	TSharedPtr<FLauncherTask> PerPlatformDeviceTask = NULL;
-	TSharedPtr<FLauncherTask> FirstPlatformBuildTask = NULL;
-	TSharedPtr<FLauncherTask> FirstPlatformCookTask = NULL;
-	TSharedPtr<FLauncherTask> FirstPlatformPackageTask = NULL;
-	TSharedPtr<FLauncherTask> FirstPlatformDeviceTask = NULL;
 
-	FLauncherUATTask::FirstTimeCompile = true;
+	// determine if there is a server platform
+	FString ServerCommand = TEXT("");
+	FString ServerPlatforms = TEXT("");
+	FString Platforms = TEXT("");
+	FString PlatformCommand = TEXT("");
+	FString OptionalParams = TEXT("");
+	for (int32 PlatformIndex = 0; PlatformIndex < InPlatforms.Num(); ++PlatformIndex)
+	{
+		// switch server and no editor platforms to the proper type
+		if (InPlatforms[PlatformIndex] == TEXT("LinuxServer"))
+		{
+			ServerPlatforms += TEXT("+Linux");
+		}
+		else if (InPlatforms[PlatformIndex] == TEXT("WindowsServer"))
+		{
+			ServerPlatforms += TEXT("+Win64");
+		}
+		else if (InPlatforms[PlatformIndex] == TEXT("LinuxNoEditor"))
+		{
+			Platforms += TEXT("+Linux");
+		}
+		else if (InPlatforms[PlatformIndex] == TEXT("WindowsNoEditor") || InPlatforms[PlatformIndex] == TEXT("Windows"))
+		{
+			Platforms += TEXT("+Win64");
+		}
+		else if (InPlatforms[PlatformIndex] == TEXT("MacNoEditor"))
+		{
+			Platforms += TEXT("+Mac");
+		}
+		else
+		{
+			Platforms += TEXT("+");
+			Platforms += InPlatforms[PlatformIndex];
+		}
 
-	// determine deployment platforms
-	ILauncherDeviceGroupPtr DeviceGroup = InProfile->GetDeployedDeviceGroup();
-	FName Variant = NAME_None;
+		// Append any extra UAT flags specified for this platform flavor
+		const PlatformInfo::FPlatformInfo* PlatformInfo = PlatformInfo::FindPlatformInfo(FName(*InPlatforms[PlatformIndex]));
+		if (!PlatformInfo->UATCommandLine.IsEmpty())
+		{
+			OptionalParams += TEXT(" ");
+			OptionalParams += PlatformInfo->UATCommandLine;
+		}
 
-	if (DeviceGroup.IsValid() && Platforms.Num() < 1)
+	}
+	if (ServerPlatforms.Len() > 0)
+	{
+		ServerCommand = TEXT(" -server -serverplatform=") + ServerPlatforms.RightChop(1);
+	}
+	if (Platforms.Len() > 0)
+	{
+		PlatformCommand = TEXT(" -platform=") + Platforms.RightChop(1);
+	}
+
+	UATCommand += PlatformCommand;
+	UATCommand += ServerCommand;
+
+	// device list
+	FString DeviceNames = TEXT("");
+	FString DeviceCommand = TEXT("");
+	ILauncherDeviceGroupPtr DeviceGroup = InProfile->GetDeployedDeviceGroup();
+	if (DeviceGroup.IsValid())
 	{
 		const TArray<FString>& Devices = DeviceGroup->GetDeviceIDs();
 		// for each deployed device...
@@ -220,418 +284,321 @@ void FLauncherWorker::CreateAndExecuteTasks( const ILauncherProfileRef& InProfil
 			if (DeviceProxy.IsValid())
 			{
 				// add the platform
-				Variant = DeviceProxy->GetTargetDeviceVariant(DeviceId);
-				Platforms.AddUnique(DeviceProxy->GetTargetPlatformName(Variant));
+				DeviceNames += TEXT("+\"") + DeviceId + TEXT("\"");
 			}			
 		}
 	}
-
-
-#if !WITH_EDITOR
-	// can't cook by the book in the editor if we are not in the editor...
-	check( InProfile->GetCookMode() != ELauncherProfileCookModes::ByTheBookInEditor );
-#endif
-
-
-	// for each desired platform...
-	for (int32 PlatformIndex = 0; PlatformIndex < Platforms.Num(); ++PlatformIndex)
+	if (DeviceNames.Len() > 0)
 	{
-		const FString& TargetPlatformName = Platforms[PlatformIndex];
-		const ITargetPlatform* TargetPlatform = GetTargetPlatformManager()->FindTargetPlatform(TargetPlatformName);
-
-		if (TargetPlatform == nullptr)
-		{
-			continue;
-		}
+		DeviceCommand += TEXT(" -device=") + DeviceNames.RightChop(1);
+	}
 
-		// ... build the editor and game ...
-		if (InProfile->IsBuilding())
+	// map list
+	FString MapList = TEXT("");
+	const TArray<FString>& CookedMaps = InProfile->GetCookedMaps();
+	if (CookedMaps.Num() > 0)
+	{
+		MapList += TEXT(" -map=");
+		for (int32 MapIndex = 0; MapIndex < CookedMaps.Num(); ++MapIndex)
 		{
-			TSharedPtr<FLauncherUATCommand> Command = NULL;
-			if (TargetPlatform->PlatformName() == TEXT("WindowsServer") || TargetPlatform->PlatformName() == TEXT("LinuxServer"))
-			{
-				Command = MakeShareable(new FLauncherBuildServerCommand(*TargetPlatform));
-			}
-			else
-			{
-				Command = MakeShareable(new FLauncherBuildGameCommand(*TargetPlatform));
-			}
-			TSharedPtr<FLauncherTask> BuildTask = MakeShareable(new FLauncherUATTask(Command, *TargetPlatform, Command->GetName(), ReadPipe, WritePipe, InProfile->GetEditorExe()));
-			BuildTask->OnStarted().AddRaw(this, &FLauncherWorker::OnTaskStarted);
-			BuildTask->OnCompleted().AddRaw(this, &FLauncherWorker::OnTaskCompleted);
-			if (!PerPlatformBuildTask.IsValid())
-			{
-				PerPlatformBuildTask = BuildTask;
-				FirstPlatformBuildTask = BuildTask;
-			}
-			else
+			MapList += CookedMaps[MapIndex];
+			if (MapIndex+1 < CookedMaps.Num())
 			{
-				PerPlatformBuildTask->AddContinuation(BuildTask);
-				PerPlatformBuildTask = BuildTask;
+				MapList += "+";
 			}
 		}
+	}
+
+	// build
+	if (InProfile->IsBuilding())
+	{
+		UATCommand += TEXT(" -build");
+
+		FCommandDesc Desc;
+		FText Command = FText::Format(LOCTEXT("LauncherBuildDesc", "Build game for {0}"), FText::FromString(Platforms.RightChop(1)));
+		Desc.Name = "Build Task";
+		Desc.Desc = Command.ToString();
+		Desc.EndText = TEXT("********** BUILD COMMAND COMPLETED **********");
+		OutCommands.Add(Desc);
+		CommandStart = TEXT("********** BUILD COMMAND STARTED **********");
+		// @todo: server
+	}
 
-		// ... cook the build...
-		TSharedPtr<FLauncherUATCommand> CookCommand = NULL;
-		if (InProfile->GetCookMode() == ELauncherProfileCookModes::ByTheBook)
+	// cook
+	switch(InProfile->GetCookMode())
+	{
+	case ELauncherProfileCookModes::ByTheBook:
 		{
-			TSharedPtr<FLauncherUATCommand> Command = NULL;
+			UATCommand += TEXT(" -cook");
 
-			if (TargetPlatform->PlatformName() == TEXT("WindowsServer") || TargetPlatform->PlatformName() == TEXT("LinuxServer"))
-			{
-				Command = MakeShareable(new FLauncherCookServerCommand(*TargetPlatform));
-			}
-			else
+			if (InProfile->IsCookingIncrementally())
 			{
-				Command = MakeShareable(new FLauncherCookGameCommand(*TargetPlatform));
+				UATCommand += TEXT(" -iterate");
 			}
-			CookCommand = Command;
-			TSharedPtr<FLauncherTask> CookTask = MakeShareable(new FLauncherUATTask(Command, *TargetPlatform, Command->GetName(), ReadPipe, WritePipe, InProfile->GetEditorExe()));
-			CookTask->OnStarted().AddRaw(this, &FLauncherWorker::OnTaskStarted);
-			CookTask->OnCompleted().AddRaw(this, &FLauncherWorker::OnTaskCompleted);
-			if (!PerPlatformCookTask.IsValid())
+
+			if (InProfile->IsCookingUnversioned())
 			{
-				PerPlatformCookTask = CookTask;
-				FirstPlatformCookTask = CookTask;
+				UATCommand += TEXT(" -Unversioned");
 			}
-			else
+
+			FString additionalOptions = InProfile->GetCookOptions();
+			if (!additionalOptions.IsEmpty())
 			{
-				PerPlatformCookTask->AddContinuation(CookTask);
-				PerPlatformCookTask = CookTask;
+				UATCommand += TEXT(" ");
+				UATCommand += additionalOptions;
 			}
-		}
-        // ... start a per-platform file server, if necessary...
-        else if (InProfile->GetCookMode() == ELauncherProfileCookModes::OnTheFly)
-        {
-            TSharedPtr<FLauncherTask> FileServerTask = NULL; //PerPlatformFileServerTasks.FindOrAdd(TargetPlatformName);
-            
-            if (!FileServerTask.IsValid())
-            {
-                TSharedPtr<FLauncherUATCommand> Command = NULL;
-                if (InProfile->GetLaunchMode() == ELauncherProfileLaunchModes::DoNotLaunch)
-                {
-                    Command = MakeShareable(new FLauncherStandAloneCookOnTheFlyCommand(*TargetPlatform));
-                }
-                else
-                {
-                    Command = MakeShareable(new FLauncherCookOnTheFlyCommand(*TargetPlatform));
-                }
-                CookCommand = Command;
-            }
-        }
-		else if ( InProfile->GetCookMode() == ELauncherProfileCookModes::ByTheBookInEditor )
-		{
-			// need a command which will wait for the cook to finish
-			class FWaitForCookInEditorToFinish : public FLauncherTask
-			{
-			public:
-				FWaitForCookInEditorToFinish() : FLauncherTask( FString(TEXT("CookByTheBookInEditor")), FString(TEXT("CookByTheBookInEditorDesk")), NULL, NULL)
-				{
-				}
-				virtual bool PerformTask( FLauncherTaskChainState& ChainState ) override
-				{
-					while ( !ChainState.Profile->OnIsCookFinished().Execute() )
-					{
-						FPlatformProcess::Sleep( 0.1f );
-					}
-					return true;
-				}
-			};
-
-			TSharedPtr<FLauncherTask> Command = MakeShareable( new FWaitForCookInEditorToFinish() );
 
-			if (!PerPlatformCookTask.IsValid())
+			if (InProfile->IsPackingWithUnrealPak())
 			{
-				PerPlatformCookTask = Command;
-				FirstPlatformCookTask = Command;
+				UATCommand += TEXT(" -pak");
 			}
-			else
+
+			FCommandDesc Desc;
+			FText Command = FText::Format(LOCTEXT("LauncherCookDesc", "Cook content for {0}"), FText::FromString(Platforms.RightChop(1)));
+			Desc.Name = "Cook Task";
+			Desc.Desc = Command.ToString();
+			Desc.EndText = TEXT("********** COOK COMMAND COMPLETED **********");
+			OutCommands.Add(Desc);
+			if (CommandStart.Len() == 0)
 			{
-				PerPlatformCookTask->AddContinuation(Command);
-				PerPlatformCookTask = Command;
+				CommandStart = TEXT("********** COOK COMMAND STARTED **********");
 			}
 		}
-		// ... package the build...
-		if (InProfile->GetPackagingMode() != ELauncherProfilePackagingModes::DoNotPackage || ((TargetPlatform->PlatformName() == TEXT("IOS") || TargetPlatform->PlatformName() == TEXT("HTML5")) && InProfile->GetDeploymentMode() != ELauncherProfileDeploymentModes::CopyRepository && InProfile->GetDeploymentMode() != ELauncherProfileDeploymentModes::DoNotDeploy))
+		break;
+	case ELauncherProfileCookModes::OnTheFly:
 		{
-			TSharedPtr<FLauncherUATCommand> Command = NULL;
-			if (TargetPlatform->PlatformName() == TEXT("WindowsServer") || TargetPlatform->PlatformName() == TEXT("LinuxServer"))
-			{
-				Command = MakeShareable(new FLauncherPackageServerCommand(*TargetPlatform, CookCommand));
-			}
-			else
-			{
-				Command = MakeShareable(new FLauncherPackageGameCommand(*TargetPlatform, CookCommand));
-			}
-			
-			TSharedPtr<FLauncherTask> PackageTask = MakeShareable(new FLauncherUATTask(Command, *TargetPlatform, Command->GetName(), ReadPipe, WritePipe, InProfile->GetEditorExe()));
-			PackageTask->OnStarted().AddRaw(this, &FLauncherWorker::OnTaskStarted);
-			PackageTask->OnCompleted().AddRaw(this, &FLauncherWorker::OnTaskCompleted);
-			if (!PerPlatformPackageTask.IsValid())
+			UATCommand += TEXT(" -cookonthefly");
+
+//			if (InProfile->GetDeploymentMode() == ELauncherProfileDeploymentModes::DoNotDeploy)
 			{
-				PerPlatformPackageTask = PackageTask;
-				FirstPlatformPackageTask = PackageTask;
+				UATCommand += " -nokill";
 			}
-			else
+
+			FCommandDesc Desc;
+			FText Command = LOCTEXT("LauncherCookDesc", "Starting cook on the fly server");
+			Desc.Name = "Cook Server Task";
+			Desc.Desc = Command.ToString();
+			Desc.EndText = TEXT("********** COOK COMMAND COMPLETED **********");
+			OutCommands.Add(Desc);
+			if (CommandStart.Len() == 0)
 			{
-				PerPlatformPackageTask->AddContinuation(PackageTask);
-				PerPlatformPackageTask = PackageTask;
+				CommandStart = TEXT("********** COOK COMMAND STARTED **********");
 			}
 		}
-		else if (InProfile->GetDeploymentMode() != ELauncherProfileDeploymentModes::DoNotDeploy && InProfile->GetDeploymentMode() != ELauncherProfileDeploymentModes::CopyRepository)
+		break;
+	case ELauncherProfileCookModes::ByTheBookInEditor:
+	case ELauncherProfileCookModes::DoNotCook:
+		UATCommand += TEXT(" -skipcook");
+		break;
+	}
+
+	// stage/package/deploy
+	if (InProfile->GetDeploymentMode() != ELauncherProfileDeploymentModes::DoNotDeploy)
+	{
+		switch (InProfile->GetDeploymentMode())
 		{
-			TSharedPtr<FLauncherUATCommand> Command = NULL;
-			if (TargetPlatform->PlatformName() == TEXT("WindowsServer") || TargetPlatform->PlatformName() == TEXT("LinuxServer"))
-			{
-				Command = MakeShareable(new FLauncherStageServerCommand(*TargetPlatform, CookCommand));
-			}
-			else
+		case ELauncherProfileDeploymentModes::CopyRepository:
+			break;
+		case ELauncherProfileDeploymentModes::FileServer:
+		case ELauncherProfileDeploymentModes::CopyToDevice:
 			{
-				Command = MakeShareable(new FLauncherStageGameCommand(*TargetPlatform, CookCommand));
+				UATCommand += TEXT(" -stage -deploy");
+				UATCommand += CommandLine;
+				UATCommand += StageDirectory;
+				UATCommand += DeviceCommand;
+				UATCommand += AdditionalCommandLine;
+
+				FCommandDesc Desc;
+				FText Command = FText::Format(LOCTEXT("LauncherDeployDesc", "Deploying content for {0}"), FText::FromString(Platforms.RightChop(1)));
+				Desc.Name = "Deploy Task";
+				Desc.Desc = Command.ToString();
+				Desc.EndText = TEXT("********** DEPLOY COMMAND COMPLETED **********");
+				OutCommands.Add(Desc);
+				if (CommandStart.Len() == 0)
+				{
+					CommandStart = TEXT("********** STAGE COMMAND STARTED **********");
+				}
 			}
+			break;
+		}
 
-			TSharedPtr<FLauncherTask> StageTask = MakeShareable(new FLauncherUATTask(Command, *TargetPlatform, Command->GetName(), ReadPipe, WritePipe, InProfile->GetEditorExe()));
-			StageTask->OnStarted().AddRaw(this, &FLauncherWorker::OnTaskStarted);
-			StageTask->OnCompleted().AddRaw(this, &FLauncherWorker::OnTaskCompleted);
-			if (!PerPlatformPackageTask.IsValid())
-			{
-				PerPlatformPackageTask = StageTask;
-				FirstPlatformPackageTask = StageTask;
-			}
-			else
+		// run
+		if (InProfile->GetLaunchMode() != ELauncherProfileLaunchModes::DoNotLaunch)
+		{
+			UATCommand += TEXT(" -run");
+
+			FCommandDesc Desc;
+			FText Command = FText::Format(LOCTEXT("LauncherRunDesc", "Launching on {0}"), FText::FromString(DeviceNames.RightChop(1)));
+			Desc.Name = "Run Task";
+			Desc.Desc = Command.ToString();
+			Desc.EndText = TEXT("********** RUN COMMAND COMPLETED **********");
+			OutCommands.Add(Desc);
+			if (CommandStart.Len() == 0)
 			{
-				PerPlatformPackageTask->AddContinuation(StageTask);
-				PerPlatformPackageTask = StageTask;
+				CommandStart = TEXT("********** RUN COMMAND STARTED **********");
 			}
 		}
-		
-		// ... and deploy the build
-		if (InProfile->GetDeploymentMode() != ELauncherProfileDeploymentModes::DoNotDeploy)
+	}
+	else
+	{
+		if (InProfile->GetPackagingMode() == ELauncherProfilePackagingModes::Locally)
 		{
-			if (!DeviceGroup.IsValid())
+			UATCommand += TEXT(" -stage -package");
+			UATCommand += StageDirectory;
+			UATCommand += CommandLine;
+			UATCommand += AdditionalCommandLine;
+
+			FCommandDesc Desc;
+			FText Command = FText::Format(LOCTEXT("LauncherPackageDesc", "Packaging content for {0}"), FText::FromString(Platforms.RightChop(1)));
+			Desc.Name = "Package Task";
+			Desc.Desc = Command.ToString();
+			Desc.EndText = TEXT("********** PACKAGE COMMAND COMPLETED **********");
+			OutCommands.Add(Desc);
+			if (CommandStart.Len() == 0)
 			{
-				continue;
+				CommandStart = TEXT("********** STAGE COMMAND STARTED **********");
 			}
+		}
+	}
 
-			TMap<FString, TSharedPtr<FLauncherTask> > PerPlatformFileServerTasks;
+	// wait for completion of UAT
+	FCommandDesc Desc;
+	FText Command = LOCTEXT("LauncherCompletionDesc", "UAT post launch cleanup");
+	Desc.Name = "Post Launch Task";
+	Desc.Desc = Command.ToString();
+	Desc.EndText = TEXT("********** LAUNCH COMPLETED **********");
+	OutCommands.Add(Desc);
 
-			const TArray<FString>& Devices = DeviceGroup->GetDeviceIDs();
+	return UATCommand;
+}
 
-			// for each deployed device...
-			for (int32 DeviceIndex = 0; DeviceIndex < Devices.Num(); ++DeviceIndex)
-			{
-				const FString& DeviceId = Devices[DeviceIndex];
-				
-				// @Todo: this can probably been cleaned up if the variant class is exposed to the user (not 100% sure i want that yet)
-				//			Deeper investigation / refactoring may make this validation obsolete.
-				ITargetDeviceProxyPtr DeviceProxyPtr = DeviceProxyManager->FindProxyDeviceForTargetDevice(DeviceId);
-				if (!DeviceProxyPtr.IsValid() || DeviceProxyPtr->GetTargetPlatformName(DeviceProxyPtr->GetTargetDeviceVariant(DeviceId)) != TargetPlatformName)
-				{
-					continue;
-				}
-				ITargetDeviceProxyRef DeviceProxy = DeviceProxyPtr.ToSharedRef();
+/* FLauncherWorker implementation
+ *****************************************************************************/
 
-				TSharedPtr<FLauncherTask> PerDeviceTask;
+void FLauncherWorker::CreateAndExecuteTasks( const ILauncherProfileRef& InProfile )
+{
+	FPlatformProcess::CreatePipe(ReadPipe, WritePipe);
 
-				FString LaunchCommandLine = InProfile->GetLaunchMode() != ELauncherProfileLaunchModes::DoNotLaunch ? InProfile->GetDefaultLaunchRole()->GetCommandLine() : TEXT("");
-				// ... start a per-platform file server, if necessary...
-				if (InProfile->GetDeploymentMode() == ELauncherProfileDeploymentModes::FileServer)
-				{
-					TSharedPtr<FLauncherTask> FileServerTask = NULL; //PerPlatformFileServerTasks.FindOrAdd(TargetPlatformName);
+	// create task chains
+	TaskChain = MakeShareable(new FLauncherVerifyProfileTask());
+	TArray<FString> Platforms;
+	if (InProfile->GetCookMode() == ELauncherProfileCookModes::ByTheBook || InProfile->IsBuilding())
+	{
+		Platforms = InProfile->GetCookedPlatforms();
+	}
 
-					if (!FileServerTask.IsValid())
-					{
-						TSharedPtr<FLauncherUATCommand> Command = NULL;
-						if (InProfile->GetLaunchMode() == ELauncherProfileLaunchModes::DoNotLaunch)
-						{
-							Command = MakeShareable(new FLauncherStandAloneCookOnTheFlyCommand(*TargetPlatform));
-						}
-						else
-						{
-							Command = MakeShareable(new FLauncherCookOnTheFlyCommand(*TargetPlatform));
-						}
-						CookCommand = Command;
-						FileServerTask = MakeShareable(new FLauncherUATTask(Command, *TargetPlatform, Command->GetName(), ReadPipe, WritePipe, InProfile->GetEditorExe()));
-						FileServerTask->OnStarted().AddRaw(this, &FLauncherWorker::OnTaskStarted);
-						FileServerTask->OnCompleted().AddRaw(this, &FLauncherWorker::OnTaskCompleted);
-					}
+	FLauncherUATTask::FirstTimeCompile = true;
 
-					if (!PerPlatformDeviceTask.IsValid())
-					{
-						PerPlatformDeviceTask = FileServerTask;
-						FirstPlatformDeviceTask = FileServerTask;
-					}
-					else
-					{
-						PerPlatformDeviceTask->AddContinuation(FileServerTask);
-						PerPlatformDeviceTask = FileServerTask;
-					}
-				}
+	// determine deployment platforms
+	ILauncherDeviceGroupPtr DeviceGroup = InProfile->GetDeployedDeviceGroup();
+	FName Variant = NAME_None;
 
-				// ... deploy the build to the device...
-				if (InProfile->GetDeploymentMode() == ELauncherProfileDeploymentModes::CopyToDevice)
-				{
-					TSharedPtr<FLauncherUATCommand> Command = NULL;
-					if (TargetPlatform->PlatformName() == TEXT("WindowsServer") || TargetPlatform->PlatformName() == TEXT("LinuxServer"))
-					{
-						Command = MakeShareable(new FLauncherDeployServerToDeviceCommand(DeviceProxy, Variant, *TargetPlatform, CookCommand));
-					}
-					else
-					{
-						Command = MakeShareable(new FLauncherDeployGameToDeviceCommand(DeviceProxy, Variant, *TargetPlatform, CookCommand, LaunchCommandLine));
-					}
-					TSharedPtr<FLauncherTask> DeployTask = MakeShareable(new FLauncherUATTask(Command, *TargetPlatform, Command->GetName(), ReadPipe, WritePipe, InProfile->GetEditorExe()));
-					DeployTask->OnStarted().AddRaw(this, &FLauncherWorker::OnTaskStarted);
-					DeployTask->OnCompleted().AddRaw(this, &FLauncherWorker::OnTaskCompleted);
+	if (DeviceGroup.IsValid() && Platforms.Num() < 1)
+	{
+		const TArray<FString>& Devices = DeviceGroup->GetDeviceIDs();
+		// for each deployed device...
+		for (int32 DeviceIndex = 0; DeviceIndex < Devices.Num(); ++DeviceIndex)
+		{
+			const FString& DeviceId = Devices[DeviceIndex];
 
-					if (!PerPlatformDeviceTask.IsValid())
-					{
-						PerPlatformDeviceTask = DeployTask;
-						FirstPlatformDeviceTask = DeployTask;
-					}
-					else
-					{
-						PerPlatformDeviceTask->AddContinuation(DeployTask);
-						PerPlatformDeviceTask = DeployTask;
-					}
-				}
-				else if (InProfile->GetDeploymentMode() == ELauncherProfileDeploymentModes::CopyRepository)
-				{
-					TSharedPtr<FLauncherUATCommand> Command = NULL;
-					if (TargetPlatform->PlatformName() == TEXT("WindowsServer") || TargetPlatform->PlatformName() == TEXT("LinuxServer"))
-					{
-						Command = MakeShareable(new FLauncherDeployServerPackageToDeviceCommand(DeviceProxy, Variant, *TargetPlatform, CookCommand));
-					}
-					else
-					{
-						Command = MakeShareable(new FLauncherDeployGamePackageToDeviceCommand(DeviceProxy, Variant, *TargetPlatform, CookCommand, LaunchCommandLine));
-					}
-					TSharedPtr<FLauncherTask> DeployTask = MakeShareable(new FLauncherUATTask(Command, *TargetPlatform, Command->GetName(), ReadPipe, WritePipe, InProfile->GetEditorExe()));
-					DeployTask->OnStarted().AddRaw(this, &FLauncherWorker::OnTaskStarted);
-					DeployTask->OnCompleted().AddRaw(this, &FLauncherWorker::OnTaskCompleted);
+			ITargetDeviceProxyPtr DeviceProxy = DeviceProxyManager->FindProxyDeviceForTargetDevice(DeviceId);
 
-					if (!PerPlatformDeviceTask.IsValid())
-					{
-						PerPlatformDeviceTask = DeployTask;
-						FirstPlatformDeviceTask = DeployTask;
-					}
-					else
-					{
-						PerPlatformDeviceTask->AddContinuation(DeployTask);
-						PerPlatformDeviceTask = DeployTask;
-					}
-				}
-				else if (TargetPlatform->PlatformName() == TEXT("XboxOne") || TargetPlatform->PlatformName() == TEXT("IOS") || TargetPlatform->PlatformName().StartsWith(TEXT("Android")))
-				{
-					TSharedPtr<FLauncherUATCommand> Command = MakeShareable(new FLauncherDeployGameToDeviceCommand(DeviceProxy, Variant, *TargetPlatform, CookCommand, LaunchCommandLine));
-					TSharedPtr<FLauncherTask> DeployTask = MakeShareable(new FLauncherUATTask(Command, *TargetPlatform, Command->GetName(), ReadPipe, WritePipe, InProfile->GetEditorExe()));
-					DeployTask->OnStarted().AddRaw(this, &FLauncherWorker::OnTaskStarted);
-					DeployTask->OnCompleted().AddRaw(this, &FLauncherWorker::OnTaskCompleted);
+			if (DeviceProxy.IsValid())
+			{
+				// add the platform
+				Variant = DeviceProxy->GetTargetDeviceVariant(DeviceId);
+				Platforms.AddUnique(DeviceProxy->GetTargetPlatformName(Variant));
+			}			
+		}
+	}
 
-					if (!PerPlatformDeviceTask.IsValid())
-					{
-						PerPlatformDeviceTask = DeployTask;
-						FirstPlatformDeviceTask = DeployTask;
-					}
-					else
-					{
-						PerPlatformDeviceTask->AddContinuation(DeployTask);
-						PerPlatformDeviceTask = DeployTask;
-					}
-				}
+#if !WITH_EDITOR
+	// can't cook by the book in the editor if we are not in the editor...
+	check( InProfile->GetCookMode() != ELauncherProfileCookModes::ByTheBookInEditor );
+#endif
 
-				// ... and then launch the build...
-				if (InProfile->GetLaunchMode() != ELauncherProfileLaunchModes::DoNotLaunch)
+	TSharedPtr<FLauncherTask> NextTask = TaskChain;
+	if (InProfile->GetCookMode() == ELauncherProfileCookModes::ByTheBookInEditor)
+	{
+		// need a command which will wait for the cook to finish
+		class FWaitForCookInEditorToFinish : public FLauncherTask
+		{
+		public:
+			FWaitForCookInEditorToFinish() : FLauncherTask( FString(TEXT("CookByTheBookInEditor")), FString(TEXT("CookByTheBookInEditorDesk")), NULL, NULL)
+			{
+			}
+			virtual bool PerformTask( FLauncherTaskChainState& ChainState ) override
+			{
+				while ( !ChainState.Profile->OnIsCookFinished().Execute() )
 				{
-					TArray<ILauncherProfileLaunchRolePtr> Roles;
-
-					if (InProfile->GetLaunchRolesFor(DeviceId, Roles) > 0)
-					{
-						// ... for every role assigned to that device
-						for (int32 RoleIndex = 0; RoleIndex < Roles.Num(); ++RoleIndex)
-						{
-							TSharedPtr<FLauncherUATCommand> Command;
-							if (Roles[RoleIndex]->GetInstanceType() == ELauncherProfileRoleInstanceTypes::StandaloneClient)
-							{
-								Command = MakeShareable(new FLauncherLaunchGameCommand(DeviceProxy, Variant, *TargetPlatform, Roles[RoleIndex].ToSharedRef(), CookCommand));
-							}
-							else if (Roles[RoleIndex]->GetInstanceType() == ELauncherProfileRoleInstanceTypes::DedicatedServer)
-							{
-								Command = MakeShareable(new FLauncherLaunchDedicatedServerCommand(DeviceProxy, Variant, *TargetPlatform, Roles[RoleIndex].ToSharedRef(), CookCommand));
-							}
-								
-							TSharedPtr<FLauncherTask> PerRoleTask = MakeShareable(new FLauncherUATTask(Command, *TargetPlatform, Command->GetName(), ReadPipe, WritePipe, InProfile->GetEditorExe()));
-							PerRoleTask->OnStarted().AddRaw(this, &FLauncherWorker::OnTaskStarted);
-							PerRoleTask->OnCompleted().AddRaw(this, &FLauncherWorker::OnTaskCompleted);
-
-							if (!PerPlatformDeviceTask.IsValid())
-							{
-								PerPlatformDeviceTask = PerRoleTask;
-								FirstPlatformDeviceTask = PerRoleTask;
-							}
-							else
-							{
-								PerPlatformDeviceTask->AddContinuation(PerRoleTask);
-								PerPlatformDeviceTask = PerRoleTask;
-							}
-						}
-					}
+					FPlatformProcess::Sleep( 0.1f );
 				}
+				return true;
 			}
-		}
-		else if (InProfile->GetCookMode() == ELauncherProfileCookModes::OnTheFly)
+		};
+		TSharedPtr<FLauncherTask> WaitTask = MakeShareable(new FWaitForCookInEditorToFinish());
+		WaitTask->OnStarted().AddRaw(this, &FLauncherWorker::OnTaskStarted);
+		WaitTask->OnCompleted().AddRaw(this, &FLauncherWorker::OnTaskCompleted);
+		NextTask->AddContinuation(WaitTask);
+		NextTask = WaitTask;
+	}
+	TArray<FCommandDesc> Commands;
+	FString StartString;
+	FString UATCommand = CreateUATCommand(InProfile, Platforms, Commands, StartString);
+	TSharedPtr<FLauncherTask> BuildTask = MakeShareable(new FLauncherUATTask(UATCommand, TEXT("Build Task"), TEXT("Launching UAT..."), ReadPipe, WritePipe, InProfile->GetEditorExe(), ProcHandle, this, StartString));
+	BuildTask->OnStarted().AddRaw(this, &FLauncherWorker::OnTaskStarted);
+	BuildTask->OnCompleted().AddRaw(this, &FLauncherWorker::OnTaskCompleted);
+	NextTask->AddContinuation(BuildTask);
+	NextTask = BuildTask;
+	for (int32 Index = 0; Index < Commands.Num(); ++Index)
+	{
+		class FLauncherWaitTask : public FLauncherTask
 		{
-//			TMap<FString, TSharedPtr<FLauncherTask> > PerPlatformFileServerTasks;
-			TSharedPtr<FLauncherTask> FileServerTask = NULL; //PerPlatformFileServerTasks.FindOrAdd(TargetPlatformName);
-
-			if (!FileServerTask.IsValid())
+		public:
+			FLauncherWaitTask( const FString& InCommandEnd, const FString& InName, const FString& InDesc, FProcHandle& InProcessHandle, ILauncherWorker* InWorker)
+				: FLauncherTask(InName, InDesc, 0, 0)
+				, CommandText(InCommandEnd)
+				, ProcessHandle(InProcessHandle)
 			{
-				TSharedPtr<FLauncherUATCommand> Command = MakeShareable(new FLauncherStandAloneCookOnTheFlyCommand(*TargetPlatform));
-				FileServerTask = MakeShareable(new FLauncherUATTask(Command, *TargetPlatform, Command->GetName(), ReadPipe, WritePipe, InProfile->GetEditorExe()));
-				FileServerTask->OnStarted().AddRaw(this, &FLauncherWorker::OnTaskStarted);
-				FileServerTask->OnCompleted().AddRaw(this, &FLauncherWorker::OnTaskCompleted);
+				EndTextFound = false;
+				InWorker->OnOutputReceived().AddRaw(this, &FLauncherWaitTask::HandleOutputReceived);
 			}
 
-			if (!PerPlatformDeviceTask.IsValid())
+		protected:
+			virtual bool PerformTask( FLauncherTaskChainState& ChainState ) override
 			{
-				PerPlatformDeviceTask = FileServerTask;
-				FirstPlatformDeviceTask = FileServerTask;
+				while (FPlatformProcess::IsProcRunning(ProcessHandle) && !EndTextFound)
+				{
+					if (GetStatus() == ELauncherTaskStatus::Canceling)
+					{
+						FPlatformProcess::TerminateProc(ProcessHandle, true);
+						return false;
+					}
+					FPlatformProcess::Sleep(0.25);
+				}
+				if (!EndTextFound && !FPlatformProcess::GetProcReturnCode(ProcessHandle, &Result))
+				{
+					return false;
+				}
+				return (Result == 0);
 			}
-			else
+
+			void HandleOutputReceived(const FString& InMessage)
 			{
-				PerPlatformDeviceTask->AddContinuation(FileServerTask);
-				PerPlatformDeviceTask = FileServerTask;
+				EndTextFound |= InMessage.Contains(CommandText);
 			}
-		}
-	}
 
-	if (FirstPlatformBuildTask.IsValid())
-	{
-		TaskChain->AddContinuation(FirstPlatformBuildTask);
-	}
-	else
-	{
-		PerPlatformBuildTask = TaskChain;
-	}
-	if (FirstPlatformCookTask.IsValid())
-	{
-		PerPlatformBuildTask->AddContinuation(FirstPlatformCookTask);
-	}
-	else
-	{
-		PerPlatformCookTask = PerPlatformBuildTask;
-	}
-	if (FirstPlatformPackageTask.IsValid())
-	{
-		PerPlatformCookTask->AddContinuation(FirstPlatformPackageTask);
-	}
-	else
-	{
-		PerPlatformPackageTask = PerPlatformCookTask;
-	}
-	if (FirstPlatformDeviceTask.IsValid())
-	{
-		PerPlatformPackageTask->AddContinuation(FirstPlatformDeviceTask);
+		private:
+			FString CommandText;
+			FProcHandle& ProcessHandle;
+			bool EndTextFound;
+		};			
+
+		TSharedPtr<FLauncherTask> WaitTask = MakeShareable(new FLauncherWaitTask(Commands[Index].EndText, Commands[Index].Name, Commands[Index].Desc, ProcHandle, this));
+		WaitTask->OnStarted().AddRaw(this, &FLauncherWorker::OnTaskStarted);
+		WaitTask->OnCompleted().AddRaw(this, &FLauncherWorker::OnTaskCompleted);
+		NextTask->AddContinuation(WaitTask);
+		NextTask = WaitTask;
 	}
 
 	// execute the chain
@@ -642,3 +609,5 @@ void FLauncherWorker::CreateAndExecuteTasks( const ILauncherProfileRef& InProfil
 
 	TaskChain->Execute(ChainState);
 }
+
+#undef LOCTEXT_NAMESPACE
diff --git a/Engine/Source/Developer/LauncherServices/Private/Launcher/LauncherWorker.h b/Engine/Source/Developer/LauncherServices/Private/Launcher/LauncherWorker.h
--- a/Engine/Source/Developer/LauncherServices/Private/Launcher/LauncherWorker.h
+++ b/Engine/Source/Developer/LauncherServices/Private/Launcher/LauncherWorker.h
@@ -3,6 +3,13 @@
 #pragma once
 
 
+struct FCommandDesc
+{
+	FString Name;
+	FString Desc;
+	FString EndText;
+};
+
 /**
  * Implements the launcher's worker thread.
  */
@@ -84,6 +91,8 @@ protected:
 	void OnTaskStarted(const FString& TaskName);
 	void OnTaskCompleted(const FString& TaskName);
 
+	FString CreateUATCommand( const ILauncherProfileRef& InProfile, const TArray<FString>& InPlatforms, TArray<FCommandDesc>& OutCommands, FString& CommandStart );
+
 private:
 
 	// Holds a critical section to lock access to selected properties.
@@ -101,9 +110,10 @@ private:
 	// Holds the first task in the task chain.
 	TSharedPtr<FLauncherTask> TaskChain;
 
-	// holds the read and write pipes
+	// holds the read and write pipes and the running UAT process
 	void* ReadPipe;
 	void* WritePipe;
+	FProcHandle ProcHandle;
 
 	double StageStartTime;
 	double LaunchStartTime;
diff --git a/Engine/Source/Developer/LauncherServices/Private/LauncherServicesPrivatePCH.h b/Engine/Source/Developer/LauncherServices/Private/LauncherServicesPrivatePCH.h
--- a/Engine/Source/Developer/LauncherServices/Private/LauncherServicesPrivatePCH.h
+++ b/Engine/Source/Developer/LauncherServices/Private/LauncherServicesPrivatePCH.h
@@ -37,13 +37,7 @@
 // launcher worker
 #include "LauncherTaskChainState.h"
 #include "LauncherTask.h"
-#include "LauncherUATCommand.h"
 #include "LauncherUATTask.h"
-#include "LauncherBuildCommands.h"
-#include "LauncherCookCommands.h"
-#include "LauncherDeployCommands.h"
-#include "LauncherLaunchCommands.h"
-#include "LauncherPackageCommands.h"
 #include "LauncherVerifyProfileTask.h"
 #include "LauncherWorker.h"
 #include "Launcher.h"
diff --git a/Engine/Source/Programs/AutomationTool/IOS/IOSPlatform.Automation.cs b/Engine/Source/Programs/AutomationTool/IOS/IOSPlatform.Automation.cs
--- a/Engine/Source/Programs/AutomationTool/IOS/IOSPlatform.Automation.cs
+++ b/Engine/Source/Programs/AutomationTool/IOS/IOSPlatform.Automation.cs
@@ -920,6 +920,11 @@ public class IOSPlatform : Platform
 		get { return false; }
 	}
 
+	public override bool RequiresPackageToDeploy
+	{
+		get { return true; }
+	}
+
 	#region Hooks
 
 	public override void PreBuildAgenda(UE4Build Build, UE4Build.BuildAgenda Agenda)
diff --git a/Engine/Source/Programs/AutomationTool/Platform.cs b/Engine/Source/Programs/AutomationTool/Platform.cs
--- a/Engine/Source/Programs/AutomationTool/Platform.cs
+++ b/Engine/Source/Programs/AutomationTool/Platform.cs
@@ -351,14 +351,6 @@ namespace AutomationTool
 			get { return false; }
 		}
 
-		/// <summary>
-		/// True if this platform requires UFE for packaging
-		/// </summary>
-		public virtual bool PackageViaUFE
-		{
-			get { return false; }
-		}
-
 		/// <summary>
 		/// True if this platform requires UFE for deploying
 		/// </summary>
@@ -399,6 +391,14 @@ namespace AutomationTool
             return "";
         }
 
+		/// <summary>
+		///  Returns whether the platform requires a package to deploy to a device
+		/// </summary>
+		public virtual bool RequiresPackageToDeploy
+		{
+			get { return false; }
+		}
+
 		#region Hooks
 
 		public virtual void PreBuildAgenda(UE4Build Build, UE4Build.BuildAgenda Agenda)
diff --git a/Engine/Source/Programs/AutomationTool/ProcessUtils.cs b/Engine/Source/Programs/AutomationTool/ProcessUtils.cs
--- a/Engine/Source/Programs/AutomationTool/ProcessUtils.cs
+++ b/Engine/Source/Programs/AutomationTool/ProcessUtils.cs
@@ -713,11 +713,11 @@ namespace AutomationTool
 				Result.WaitForExit();
 				var BuildDuration = (DateTime.UtcNow - StartTime).TotalMilliseconds;
 				AddRunTime(App, (int)(BuildDuration));
-                if (!Options.HasFlag(ERunOptions.NoLoggingOfRunCommand))
+				Result.ExitCode = Proc.ExitCode;
+				if (!Options.HasFlag(ERunOptions.NoLoggingOfRunCommand))
                 {
                     Log(SpewVerbosity,"Run: Took {0}s to run {1}, ExitCode={2}", BuildDuration / 1000, Path.GetFileName(App), Result.ExitCode);
                 }
-				Result.ExitCode = Proc.ExitCode;
 				Result.OnProcessExited();
                 Result.DisposeProcess();
 			}
@@ -967,7 +967,7 @@ namespace AutomationTool
 						}
 					}
 
-					while (LogReader.EndOfStream && !LogProcess.HasExited)
+					while (LogReader.EndOfStream && !LogProcess.HasExited && bKeepReading)
 					{
 						Thread.Sleep(250);
 						// Tick the callback so that it can respond to external events
diff --git a/Engine/Source/Programs/AutomationTool/Scripts/ArchiveCommand.Automation.cs b/Engine/Source/Programs/AutomationTool/Scripts/ArchiveCommand.Automation.cs
--- a/Engine/Source/Programs/AutomationTool/Scripts/ArchiveCommand.Automation.cs
+++ b/Engine/Source/Programs/AutomationTool/Scripts/ArchiveCommand.Automation.cs
@@ -49,6 +49,9 @@ public partial class Project : CommandUtils
 		{
 			return;
 		}
+
+		Log("********** ARCHIVE COMMAND STARTED **********");
+
 		if (!Params.NoClient)
 		{
 			var DeployContextList = CreateDeploymentContext(Params, false, false);
@@ -70,6 +73,7 @@ public partial class Project : CommandUtils
 				ApplyArchiveManifest(Params, SC);
 				SC.StageTargetPlatform.ProcessArchivedProject(Params, SC);
 			}
+			Log("********** ARCHIVE COMMAND COMPLETED **********");
 		}
 	}
 
diff --git a/Engine/Source/Programs/AutomationTool/Scripts/BuildProjectCommand.Automation.cs b/Engine/Source/Programs/AutomationTool/Scripts/BuildProjectCommand.Automation.cs
--- a/Engine/Source/Programs/AutomationTool/Scripts/BuildProjectCommand.Automation.cs
+++ b/Engine/Source/Programs/AutomationTool/Scripts/BuildProjectCommand.Automation.cs
@@ -24,6 +24,7 @@ public partial class Project : CommandUtils
 
 	public static void Build(BuildCommand Command, ProjectParams Params, int WorkingCL = -1)
 	{
+
 		Params.ValidateAndLog();
 
 		if (!Params.Build)
@@ -31,6 +32,8 @@ public partial class Project : CommandUtils
 			return;
 		}
 
+		Log("********** BUILD COMMAND STARTED **********");
+
 		var UE4Build = new UE4Build(Command);
 		var Agenda = new UE4Build.BuildAgenda();
 		var CrashReportPlatforms = new HashSet<UnrealTargetPlatform>();
@@ -122,6 +125,8 @@ public partial class Project : CommandUtils
 			// Open files for add or edit
 			UE4Build.AddBuildProductsToChangelist(WorkingCL, UE4Build.BuildProductFiles);
 		}
+
+		Log("********** BUILD COMMAND COMPLETED **********");
 	}
 
 	#endregion
diff --git a/Engine/Source/Programs/AutomationTool/Scripts/CookCommand.Automation.cs b/Engine/Source/Programs/AutomationTool/Scripts/CookCommand.Automation.cs
--- a/Engine/Source/Programs/AutomationTool/Scripts/CookCommand.Automation.cs
+++ b/Engine/Source/Programs/AutomationTool/Scripts/CookCommand.Automation.cs
@@ -21,89 +21,123 @@ public partial class Project : CommandUtils
 
 	public static void Cook(ProjectParams Params)
 	{
-		if (!Params.Cook || Params.SkipCook)
+		if ((!Params.Cook && !(Params.CookOnTheFly && !Params.SkipServer)) || Params.SkipCook)
 		{
 			return;
 		}
 		Params.ValidateAndLog();
 
+		Log("********** COOK COMMAND STARTED **********");
+
 		string UE4EditorExe = HostPlatform.Current.GetUE4ExePath(Params.UE4Exe);
 		if (!FileExists(UE4EditorExe))
 		{
 			throw new AutomationException("Missing " + UE4EditorExe + " executable. Needs to be built first.");
 		}
 
-		var PlatformsToCook = new HashSet<string>();
-
-		if (!Params.NoClient)
+		if (Params.CookOnTheFly && !Params.SkipServer)
 		{
-			foreach (var ClientPlatform in Params.ClientTargetPlatforms)
+			if (Params.ClientTargetPlatforms.Count > 0)
 			{
-				// Use the data platform, sometimes we will copy another platform's data
-				var DataPlatform = Params.GetCookedDataPlatformForClientTarget(ClientPlatform);
-				PlatformsToCook.Add(Params.GetTargetPlatformInstance(DataPlatform).GetCookPlatform(false, Params.HasDedicatedServerAndClient, Params.CookFlavor));
+				var LogFolderOutsideOfSandbox = GetLogFolderOutsideOfSandbox();
+				if (!GlobalCommandLine.Installed)
+				{
+					// In the installed runs, this is the same folder as CmdEnv.LogFolder so delete only in not-installed
+					DeleteDirectory(LogFolderOutsideOfSandbox);
+					CreateDirectory(LogFolderOutsideOfSandbox);
+				}
+				var ServerLogFile = CombinePaths(LogFolderOutsideOfSandbox, "Server.log");
+				Platform ClientPlatformInst = Params.ClientTargetPlatformInstances[0];
+				string TargetCook = ClientPlatformInst.GetCookPlatform(false, Params.HasDedicatedServerAndClient, Params.CookFlavor);
+				ServerProcess = RunCookOnTheFlyServer(Params.RawProjectPath, Params.NoClient ? "" : ServerLogFile, TargetCook, Params.RunCommandline);
+
+				if (ServerProcess != null)
+				{
+					Log("Waiting a few seconds for the server to start...");
+					Thread.Sleep(5000);
+				}
 			}
-		}
-		if (Params.DedicatedServer)
-		{
-			foreach (var ServerPlatform in Params.ServerTargetPlatforms)
+			else
 			{
-				// Use the data platform, sometimes we will copy another platform's data
-				var DataPlatform = Params.GetCookedDataPlatformForServerTarget(ServerPlatform);
-				PlatformsToCook.Add(Params.GetTargetPlatformInstance(DataPlatform).GetCookPlatform(true, false, Params.CookFlavor));
+				throw new AutomationException("Failed to run, client target platform not specified");
 			}
 		}
-
-		if (Params.Clean.HasValue && Params.Clean.Value && !Params.IterativeCooking)
+		else
 		{
-			Log("Cleaning cooked data.");
-			CleanupCookedData(PlatformsToCook.ToList(), Params);
-		}
+			var PlatformsToCook = new HashSet<string>();
 
-		// cook the set of maps, or the run map, or nothing
-		string[] Maps = null;
-		if (Params.HasMapsToCook)
-		{
-			Maps = Params.MapsToCook.ToArray();
-		}
+			if (!Params.NoClient)
+			{
+				foreach (var ClientPlatform in Params.ClientTargetPlatforms)
+				{
+					// Use the data platform, sometimes we will copy another platform's data
+					var DataPlatform = Params.GetCookedDataPlatformForClientTarget(ClientPlatform);
+					PlatformsToCook.Add(Params.GetTargetPlatformInstance(DataPlatform).GetCookPlatform(false, Params.HasDedicatedServerAndClient, Params.CookFlavor));
+				}
+			}
+			if (Params.DedicatedServer)
+			{
+				foreach (var ServerPlatform in Params.ServerTargetPlatforms)
+				{
+					// Use the data platform, sometimes we will copy another platform's data
+					var DataPlatform = Params.GetCookedDataPlatformForServerTarget(ServerPlatform);
+					PlatformsToCook.Add(Params.GetTargetPlatformInstance(DataPlatform).GetCookPlatform(true, false, Params.CookFlavor));
+				}
+			}
 
-		string[] Dirs = null;
-		if (Params.HasDirectoriesToCook)
-		{
-			Dirs = Params.DirectoriesToCook.ToArray();
-		}
+			if (Params.Clean.HasValue && Params.Clean.Value && !Params.IterativeCooking)
+			{
+				Log("Cleaning cooked data.");
+				CleanupCookedData(PlatformsToCook.ToList(), Params);
+			}
 
-        string[] Cultures = null;
-        if (Params.HasCulturesToCook)
-        {
-            Cultures = Params.CulturesToCook.ToArray();
-        }
+			// cook the set of maps, or the run map, or nothing
+			string[] Maps = null;
+			if (Params.HasMapsToCook)
+			{
+				Maps = Params.MapsToCook.ToArray();
+			}
 
-		try
-		{
-			var CommandletParams = "-buildmachine -Unversioned -fileopenlog";
-			if (Params.UseDebugParamForEditorExe)
+			string[] Dirs = null;
+			if (Params.HasDirectoriesToCook)
 			{
-				CommandletParams += " -debug";
+				Dirs = Params.DirectoriesToCook.ToArray();
 			}
-			if (Params.Manifests)
+
+			string[] Cultures = null;
+			if (Params.HasCulturesToCook)
 			{
-				CommandletParams += " -manifests";
+				Cultures = Params.CulturesToCook.ToArray();
 			}
-			if (Params.IterativeCooking)
+
+			try
 			{
-				CommandletParams += " -iterate";
+				var CommandletParams = "-buildmachine -Unversioned -fileopenlog";
+				if (Params.UseDebugParamForEditorExe)
+				{
+					CommandletParams += " -debug";
+				}
+				if (Params.Manifests)
+				{
+					CommandletParams += " -manifests";
+				}
+				if (Params.IterativeCooking)
+				{
+					CommandletParams += " -iterate";
+				}
+				CookCommandlet(Params.RawProjectPath, Params.UE4Exe, Maps, Dirs, Cultures, CombineCommandletParams(PlatformsToCook.ToArray()), CommandletParams);
+			}
+			catch (Exception Ex)
+			{
+				// Delete cooked data (if any) as it may be incomplete / corrupted.
+				Log("Cook failed. Deleting cooked data.");
+				CleanupCookedData(PlatformsToCook.ToList(), Params);
+				AutomationTool.ErrorReporter.Error("Cook failed.", (int)AutomationTool.ErrorCodes.Error_UnknownCookFailure);
+				throw Ex;
 			}
-			CookCommandlet(Params.RawProjectPath, Params.UE4Exe, Maps, Dirs, Cultures, CombineCommandletParams(PlatformsToCook.ToArray()), CommandletParams);
-		}
-		catch (Exception Ex)
-		{
-			// Delete cooked data (if any) as it may be incomplete / corrupted.
-			Log("Cook failed. Deleting cooked data.");
-			CleanupCookedData(PlatformsToCook.ToList(), Params);
-			AutomationTool.ErrorReporter.Error("Cook failed.", (int)AutomationTool.ErrorCodes.Error_UnknownCookFailure);
-			throw Ex;
 		}
+
+		Log("********** COOK COMMAND COMPLETED **********");
 	}
 
 	private static void CleanupCookedData(List<string> PlatformsToCook, ProjectParams Params)
diff --git a/Engine/Source/Programs/AutomationTool/Scripts/CopyBuildToStagingDirectory.Automation.cs b/Engine/Source/Programs/AutomationTool/Scripts/CopyBuildToStagingDirectory.Automation.cs
--- a/Engine/Source/Programs/AutomationTool/Scripts/CopyBuildToStagingDirectory.Automation.cs
+++ b/Engine/Source/Programs/AutomationTool/Scripts/CopyBuildToStagingDirectory.Automation.cs
@@ -276,6 +276,7 @@ public partial class Project : CommandUtils
 
 			if (SC.StageTargetPlatform.StageMovies)
 			{
+				SC.StageFiles(StagedFileTypeForMovies, CombinePaths(SC.LocalRoot, "Engine/Content/Movies"), "*", true, null, CombinePaths(SC.RelativeProjectRootForStage, "Engine/Content/Movies"), true, !Params.Pak);
 				SC.StageFiles(StagedFileTypeForMovies, CombinePaths(SC.ProjectRoot, "Content/Movies"), "*", true, null, CombinePaths(SC.RelativeProjectRootForStage, "Content/Movies"), true, !Params.Pak);
 			}
 
@@ -719,7 +720,7 @@ public partial class Project : CommandUtils
 							{
 								if (IP.UnicastAddresses[Index].IsDnsEligible && IP.UnicastAddresses[Index].Address.AddressFamily == System.Net.Sockets.AddressFamily.InterNetwork)
 								{
-                                    if (Params.Port != null)
+                                    if (!IsNullOrEmpty(Params.Port))
                                     {
                                         foreach (var Port in Params.Port)
                                         {
@@ -744,7 +745,13 @@ public partial class Project : CommandUtils
                                     }
                                     else
                                     {
-                                        // use default port
+										if (!FirstParam)
+										{
+											FileHostParams += "+";
+										}
+										FirstParam = false;
+										
+										// use default port
                                         FileHostParams += IP.UnicastAddresses[Index].Address.ToString();
                                     }
 
@@ -765,7 +772,7 @@ public partial class Project : CommandUtils
 							{
 								if (IP.UnicastAddresses[Index].IsDnsEligible)
 								{
-									if (Params.Port != null)
+									if (!IsNullOrEmpty(Params.Port))
 									{
 										foreach (var Port in Params.Port)
 										{
@@ -789,7 +796,13 @@ public partial class Project : CommandUtils
 									}
                                     else
                                     {
-                                        // use default port
+										if (!FirstParam)
+										{
+											FileHostParams += "+";
+										}
+										FirstParam = false;
+										
+										// use default port
                                         FileHostParams += IP.UnicastAddresses[Index].Address.ToString();
                                     }
 
@@ -799,7 +812,7 @@ public partial class Project : CommandUtils
 					}
 				}
 				const string LocalHost = "127.0.0.1";
-				if (Params.Port != null)
+				if (!IsNullOrEmpty(Params.Port))
 				{
 					foreach (var Port in Params.Port)
 					{
@@ -824,7 +837,13 @@ public partial class Project : CommandUtils
 				}
                 else
                 {
-                    // use default port
+					if (!FirstParam)
+					{
+						FileHostParams += "+";
+					}
+					FirstParam = false;
+					
+					// use default port
                     FileHostParams += LocalHost;
                 }
 				FileHostParams += " ";
@@ -954,6 +973,8 @@ public partial class Project : CommandUtils
 		{
 			Params.ValidateAndLog();
 
+			Log("********** STAGE COMMAND STARTED **********");
+
 			if (!Params.NoClient)
 			{
 				var DeployContextList = CreateDeploymentContext(Params, false, true);
@@ -975,6 +996,7 @@ public partial class Project : CommandUtils
 					ApplyStagingManifest(Params, SC);
 				}
 			}
+			Log("********** STAGE COMMAND COMPLETED **********");
 		}
 	}
 
diff --git a/Engine/Source/Programs/AutomationTool/Scripts/DeployCommand.Automation.cs b/Engine/Source/Programs/AutomationTool/Scripts/DeployCommand.Automation.cs
--- a/Engine/Source/Programs/AutomationTool/Scripts/DeployCommand.Automation.cs
+++ b/Engine/Source/Programs/AutomationTool/Scripts/DeployCommand.Automation.cs
@@ -18,6 +18,9 @@ public partial class Project : CommandUtils
 		{
 			return;
 		}
+
+		Log("********** DEPLOY COMMAND STARTED **********");
+		
 		if (!Params.NoClient)
 		{
 			var DeployContextList = CreateDeploymentContext(Params, false, false);
@@ -64,6 +67,8 @@ public partial class Project : CommandUtils
 				SC.StageTargetPlatform.Deploy(ServerParams, SC);
 			}
 		}
+
+		Log("********** DEPLOY COMMAND COMPLETED **********");
 	}
 
 	#endregion
diff --git a/Engine/Source/Programs/AutomationTool/Scripts/PackageCommand.Automation.cs b/Engine/Source/Programs/AutomationTool/Scripts/PackageCommand.Automation.cs
--- a/Engine/Source/Programs/AutomationTool/Scripts/PackageCommand.Automation.cs
+++ b/Engine/Source/Programs/AutomationTool/Scripts/PackageCommand.Automation.cs
@@ -14,52 +14,34 @@ public partial class Project : CommandUtils
 	public static void Package(ProjectParams Params, int WorkingCL=-1)
 	{
 		Params.ValidateAndLog();
-		if (!Params.Package)
+/*		if (!Params.Package)
 		{
 			return;
-		}
+		}*/
+
+		bool NeedsPackage = Params.Package;
+		List<DeploymentContext> DeployContextList = new List<DeploymentContext>();
 		if (!Params.NoClient)
 		{
-			var DeployContextList = CreateDeploymentContext(Params, false, false);
-			foreach ( var SC in DeployContextList )
-			{
-				if (SC.StageTargetPlatform.PackageViaUFE)
-				{
-					string ClientCmdLine = "-run=Package ";
-					ClientCmdLine += "-Targetplatform=" + SC.StageTargetPlatform.PlatformType.ToString() + " ";
-					ClientCmdLine += "-SourceDir=" + CombinePaths(Params.BaseStageDirectory, SC.StageTargetPlatform.PlatformType.ToString()) + " ";
-					string ClientApp = CombinePaths(CmdEnv.LocalRoot, "Engine/Binaries/Win64/UnrealFrontend.exe");
-
-					Log("Packaging via UFE:");
-					Log("\tClientCmdLine: " + ClientCmdLine + "");
-
-					//@todo UAT: Consolidate running of external applications like UFE (See 'RunProjectCommand' for other instances)
-					PushDir(Path.GetDirectoryName(ClientApp));
-					// Always start client process and don't wait for exit.
-					ProcessResult ClientProcess = Run(ClientApp, ClientCmdLine, null, ERunOptions.NoWaitForExit);
-					PopDir();
-					if (ClientProcess != null)
-					{
-						do
-						{
-							Thread.Sleep(100);
-						}
-						while (ClientProcess.HasExited == false);
-					}
-				}
-				else
-				{
-					SC.StageTargetPlatform.Package(Params, SC, WorkingCL);
-				}
-			}
+			DeployContextList.AddRange(CreateDeploymentContext(Params, false, false));
 		}
 		if (Params.DedicatedServer)
 		{
-			var DeployContextList = CreateDeploymentContext(Params, true, false);
+			DeployContextList.AddRange(CreateDeploymentContext(Params, true, false));
+		}
+		if (DeployContextList.Count > 0)
+		{
+			Log("********** PACKAGE COMMAND STARTED **********");
+
 			foreach (var SC in DeployContextList)
 			{
-				SC.StageTargetPlatform.Package(Params, SC, WorkingCL);
+				if (Params.Package || (SC.StageTargetPlatform.RequiresPackageToDeploy && Params.Deploy))
+				{
+					SC.StageTargetPlatform.Package(Params, SC, WorkingCL);
+				}
 			}
+
+			Log("********** PACKAGE COMMAND COMPLETED **********");
 		}
 	}
 
diff --git a/Engine/Source/Programs/AutomationTool/Scripts/RunProjectCommand.Automation.cs b/Engine/Source/Programs/AutomationTool/Scripts/RunProjectCommand.Automation.cs
--- a/Engine/Source/Programs/AutomationTool/Scripts/RunProjectCommand.Automation.cs
+++ b/Engine/Source/Programs/AutomationTool/Scripts/RunProjectCommand.Automation.cs
@@ -30,7 +30,12 @@ public partial class Project : CommandUtils
 	/// <summary>
 	/// Thread used to read client log file.
 	/// </summary>
-	private static Thread ClientLogReaderThread;
+	private static Thread ClientLogReaderThread = null;
+
+	/// <summary>
+	/// Process for the server, can be set by the cook command when a cook on the fly server is used
+	/// </summary>
+	public static ProcessResult ServerProcess;
 
 	#endregion
 
@@ -88,8 +93,10 @@ public partial class Project : CommandUtils
 			return;
 		}
 
+		Log("********** RUN COMMAND STARTED **********");
+
 		var LogFolderOutsideOfSandbox = GetLogFolderOutsideOfSandbox();
-		if (!GlobalCommandLine.Installed)
+		if (!GlobalCommandLine.Installed && ServerProcess == null)
 		{
 			// In the installed runs, this is the same folder as CmdEnv.LogFolder so delete only in not-installed
 			DeleteDirectory(LogFolderOutsideOfSandbox);
@@ -110,26 +117,14 @@ public partial class Project : CommandUtils
 		{
 			CopyLogsBackToLogFolder();
 		}
+
+		Log("********** RUN COMMAND COMPLETED **********");
 	}
 
 	private static void RunInternal(ProjectParams Params, string ServerLogFile, string ClientLogFile)
 	{
 		// Setup server process if required.
-		ProcessResult ServerProcess = null;
-		if (Params.CookOnTheFly && !Params.SkipServer)
-		{
-			if (Params.ClientTargetPlatforms.Count > 0)
-			{
-				Platform ClientPlatformInst = Params.ClientTargetPlatformInstances[0];
-				string TargetCook = ClientPlatformInst.GetCookPlatform(false, Params.HasDedicatedServerAndClient, Params.CookFlavor);
-				ServerProcess = RunCookOnTheFlyServer(Params.RawProjectPath, Params.NoClient ? "" : ServerLogFile, TargetCook, Params.RunCommandline);
-			}
-			else
-			{
-				throw new AutomationException("Failed to run, client target platform not specified");
-			}
-		}
-		else if (Params.DedicatedServer && !Params.SkipServer)
+		if (Params.DedicatedServer && !Params.SkipServer)
 		{
 			if (Params.ServerTargetPlatforms.Count > 0)
 			{
@@ -175,7 +170,7 @@ public partial class Project : CommandUtils
 			// Run the client.
 			if (ServerProcess != null)
 			{
-				RunClientWithServer(ServerLogFile, ServerProcess, ClientApp, ClientCmdLine, ClientRunFlags, ClientLogFile, Params);
+				RunClientWithServer(SC, ServerLogFile, ServerProcess, ClientApp, ClientCmdLine, ClientRunFlags, ClientLogFile, Params);
 			}
 			else
 			{
@@ -340,7 +335,7 @@ public partial class Project : CommandUtils
 		}
 	}
 
-	private static void RunClientWithServer(string ServerLogFile, ProcessResult ServerProcess, string ClientApp, string ClientCmdLine, ERunOptions ClientRunFlags, string ClientLogFile, ProjectParams Params)
+	private static void RunClientWithServer(List<DeploymentContext> DeployContextList, string ServerLogFile, ProcessResult ServerProcess, string ClientApp, string ClientCmdLine, ERunOptions ClientRunFlags, string ClientLogFile, ProjectParams Params)
 	{
 		ProcessResult ClientProcess = null;
 		var OtherClients = new List<ProcessResult>();
@@ -500,13 +495,16 @@ public partial class Project : CommandUtils
 					if (ClientProcess == null && (AllClientOutput.Contains("Game Engine Initialized") || AllClientOutput.Contains("Unreal Network File Server is ready")))
 					{
 						Log("Starting Client....");
-						ClientProcess = Run(ClientApp, ClientCmdLine, null, ClientRunFlags | ERunOptions.NoWaitForExit);
+						var SC = DeployContextList[0];
+						ClientProcess = SC.StageTargetPlatform.RunClient(ClientRunFlags | ERunOptions.NoWaitForExit, ClientApp, ClientCmdLine, Params);
+//						ClientProcess = Run(ClientApp, ClientCmdLine, null, ClientRunFlags | ERunOptions.NoWaitForExit);
 						if (NumClients > 1 && NumClients < 9)
 						{
 							for (int i = 1; i < NumClients; i++)
 							{
 								Log("Starting Extra Client....");
-								OtherClients.Add(Run(ClientApp, ClientCmdLine, null, ClientRunFlags | ERunOptions.NoWaitForExit));
+								ProcessResult NewClient = SC.StageTargetPlatform.RunClient(ClientRunFlags | ERunOptions.NoWaitForExit, ClientApp, ClientCmdLine, Params);
+								OtherClients.Add(NewClient);
 							}
 						}
 					}
@@ -526,7 +524,10 @@ public partial class Project : CommandUtils
 				{
 
 					Log("Client exited, stopping server....");
-					ServerProcess.StopProcess();
+					if (!GlobalCommandLine.NoKill)
+					{
+						ServerProcess.StopProcess();
+					}
 					bKeepReading = false;
 				}
 
@@ -602,7 +603,7 @@ public partial class Project : CommandUtils
 							{
 								if (IP.UnicastAddresses[Index].IsDnsEligible && IP.UnicastAddresses[Index].Address.AddressFamily == System.Net.Sockets.AddressFamily.InterNetwork)
 								{
-									if (Params.Port != null)
+									if (!IsNullOrEmpty(Params.Port))
 									{
 										foreach (var Port in Params.Port)
 										{
@@ -626,7 +627,13 @@ public partial class Project : CommandUtils
 									}
                                     else
                                     {
-                                        // use default port
+										if (!FirstParam)
+										{
+											TempCmdLine += "+";
+										}
+										FirstParam = false;
+										
+										// use default port
                                         TempCmdLine += IP.UnicastAddresses[Index].Address.ToString();
                                     }
 								}
@@ -646,7 +653,7 @@ public partial class Project : CommandUtils
 							{
 								if (IP.UnicastAddresses[Index].IsDnsEligible)
 								{
-									if (Params.Port != null)
+									if (!IsNullOrEmpty(Params.Port))
 									{
 										foreach (var Port in Params.Port)
 										{
@@ -670,7 +677,13 @@ public partial class Project : CommandUtils
 									}
                                     else
                                     {
-                                        // use default port
+										if (!FirstParam)
+										{
+											TempCmdLine += "+";
+										}
+										FirstParam = false;
+										
+										// use default port
                                         TempCmdLine += IP.UnicastAddresses[Index].Address.ToString();
                                     }
 								}
@@ -706,7 +719,13 @@ public partial class Project : CommandUtils
 				}
                 else
                 {
-                    // use default port
+					if (!FirstParam)
+					{
+						TempCmdLine += "+";
+					}
+					FirstParam = false;
+					
+					// use default port
                     TempCmdLine += LocalHost;
                 }
 				TempCmdLine += " ";
