Fix for device lists not being saved.

diff --git a/Engine/Source/Developer/LauncherServices/Private/Profiles/LauncherProfileManager.h b/Engine/Source/Developer/LauncherServices/Private/Profiles/LauncherProfileManager.h
--- a/Engine/Source/Developer/LauncherServices/Private/Profiles/LauncherProfileManager.h
+++ b/Engine/Source/Developer/LauncherServices/Private/Profiles/LauncherProfileManager.h
@@ -77,6 +77,8 @@ public:
 
 	virtual void RemoveDeviceGroup( const ILauncherDeviceGroupRef& DeviceGroup ) override;
 
+	virtual void SaveDeviceGroups() override;
+
 	virtual void RemoveSimpleProfile(const ILauncherSimpleProfileRef& SimpleProfile) override;
 
 	virtual void RemoveProfile( const ILauncherProfileRef& Profile ) override;
@@ -116,11 +118,6 @@ protected:
 	 */
 	ILauncherDeviceGroupPtr ParseDeviceGroup( const FString& GroupString );
 
-	/**
-	 * Saves all the device groups to a config file
-	 */
-	void SaveDeviceGroups();
-
 	/*
 	* Saves all simple profiles to disk.
 	*/
diff --git a/Engine/Source/Developer/LauncherServices/Public/Interfaces/ILauncherProfileManager.h b/Engine/Source/Developer/LauncherServices/Public/Interfaces/ILauncherProfileManager.h
--- a/Engine/Source/Developer/LauncherServices/Public/Interfaces/ILauncherProfileManager.h
+++ b/Engine/Source/Developer/LauncherServices/Public/Interfaces/ILauncherProfileManager.h
@@ -90,6 +90,11 @@ public:
 	 */
 	virtual void RemoveDeviceGroup( const ILauncherDeviceGroupRef& DeviceGroup ) = 0;
 
+	/**
+	 * Saves all the device groups to a config file
+	 */
+	virtual void SaveDeviceGroups() = 0;
+
 public:
 
 	/**
diff --git a/Engine/Source/Developer/ProjectLauncher/Private/Widgets/SProjectLauncherSettings.cpp b/Engine/Source/Developer/ProjectLauncher/Private/Widgets/SProjectLauncherSettings.cpp
--- a/Engine/Source/Developer/ProjectLauncher/Private/Widgets/SProjectLauncherSettings.cpp
+++ b/Engine/Source/Developer/ProjectLauncher/Private/Widgets/SProjectLauncherSettings.cpp
@@ -331,6 +331,8 @@ void SProjectLauncherSettings::HandleCloseActionExecute()
 	if (LaunchProfile.IsValid())
 	{
 		Model->GetProfileManager()->SaveProfile(LaunchProfile.ToSharedRef());
+		//@Todo: FIX! Very Heavy Handed! Will have to re factor the device group saving code, but branch is tonight and this is safe. 
+		Model->GetProfileManager()->SaveDeviceGroups();
 	}
 
 	if (OnCloseClicked.IsBound())
