Code and documentation cleanup pass

diff --git a/Engine/Source/Developer/FriendsAndChat/Private/FriendsMessageManager.h b/Engine/Source/Developer/FriendsAndChat/Private/FriendsMessageManager.h
--- a/Engine/Source/Developer/FriendsAndChat/Private/FriendsMessageManager.h
+++ b/Engine/Source/Developer/FriendsAndChat/Private/FriendsMessageManager.h
@@ -1,11 +1,8 @@
 // Copyright 1998-2014 Epic Games, Inc. All Rights Reserved.
 
-/*=============================================================================
-	FriendsMessageManager.h: Declares the FFriendsMessageManager interface.
-=============================================================================*/
-
 #pragma once
 
+
 // Enum holding the state of the Friends manager
 namespace EFriendsMessageManagerState
 {
@@ -21,11 +18,11 @@ namespace EFriendsMessageManagerState
 	};
 };
 
-/**
- * Delegate type for FriendsList updated.
- */
+
+/** Delegate type for FriendsList updated. */
 DECLARE_MULTICAST_DELEGATE( FOnMessagesUpdated )
 
+
 /**
  * Implement the Friend and Chat manager
  */
@@ -49,86 +46,85 @@ class FFriendsMessageManager
 
 public:
 
-	/**
-	 * Default constructor
-	 */
+	/** Default constructor. */
 	FFriendsMessageManager();
 	
-	/**
-	 * Default destructor
-	 */
+	/** Destructor. */
 	~FFriendsMessageManager();
 
 public:
 
-	/**
-	 * Default constructor
-	 */
+	/** */
 	void StartupManager();
 
 	/**
 	 * Init the manager
-	 * @param NotificationDelegate 	- the notification delegate
-	 * @param bInAllowLaunchGame 	- Can this app join a game directly
+	 *
+	 * @param NotificationDelegate The notification delegate.
+	 * @param bInAllowLaunchGame Can this app join a game directly.
 	 */
 	void Init( FOnFriendsNotification& NotificationDelegate, bool bInAllowLaunchGame );
 
 	/**
-	 * Start or stop message polling for messages
-	 * @param bStart - true for start - false for stop;
+	 * Start or stop message polling for messages.
+	 *
+	 * @param bStart true for start, false for stop.
 	 */
 	void SetMessagePolling( bool bStart );
 
-	/**
-	 * Logout - close any Friends windows.
-	 */
+	/** Logout and close any Friends windows. */
 	void Logout();
 
 	/**
-	 * Request a friend to be invited to a game
-	 * @param FriendID - The friend ID
+	 * Request a friend to be invited to a game.
+	 *
+	 * @param FriendID The friend ID.
 	 */
 	void InviteFriendToGame( TSharedRef< FUniqueNetId > FriendID );
 
 	/**
-	 * Request to join a friends game
-	 * @param FriendID - The friend ID
+	 * Request to join a friends game.
+	 *
+	 * @param FriendID The friend ID.
 	 */
 	void RequestJoinAGame( TSharedRef< FUniqueNetId > FriendID );
 
 	/**
-	 * Send a chat message
-	 * @param FriendID - Friend to send a message to.
- 	 * @param Message - The message content;
+	 * Send a chat message.
+	 *
+	 * @param FriendID Friend to send a message to.
+ 	 * @param Message The message content.
 	 */
 	void SendMessage( TSharedPtr< FFriendStuct > FriendID, const FText& Message );
 
-	/**
-	 * Clear game invites
-	 */
+	/** Clear game invites. */
 	void ClearGameInvites();
 	
 	/**
-	 * Set an unhandled notification. We will try again later
-	 * @param NetID - The NetID of the failed notfication
+	 * Set an unhandled notification. We will try again later.
+	 *
+	 * @param NetID The NetID of the failed notification.
 	 */
 	void SetUnhandledNotification( TSharedRef< FUniqueNetId > NetID );
 
 	/**
-	 * Get the manager state
-	 * @return - The manager state
+	 * Get the manager state.
+	 *
+	 * @return The manager state.
 	 */
 	EFriendsMessageManagerState::Type GetState();
 
 	/**
-	 * Get current chat messages
-	 * @return - A list of chat messages
+	 * Get current chat messages.
+	 *
+	 * @return A list of chat messages.
 	 */
 	TArray< TSharedPtr< FFriendsAndChatMessage > > GetChatMessages();
 
 	/**
 	 * Accessor for the Friends List updated delegate.
-	 * @return The delegate
+	 *
+	 * @return The delegate.
 	 */
 	FOnMessagesUpdated& OnChatListUpdated();
 
@@ -136,92 +132,84 @@ private:
 
 	/**
 	 * A ticker used to perform updates on the main thread.
-	 * @param Delta		The tick delta
-	 * @return	true to continue ticking
+	 *
+	 * @param Delta The tick delta.
+	 * @return true to continue ticking.
 	 */
 	bool Tick( float Delta );
 
 	/**
-	 * Set the manager state
-	 * @param NewState - The new manager state
+	 * Set the manager state.
+	 *
+	 * @param NewState The new manager state.
 	 */
 	void SetState( EFriendsMessageManagerState::Type NewState );
 
-	/**
-	 * Enumerate messages.
-	 */
+	/** Enumerate messages. */
 	void RequestEnumerateMessages();
 
-	/**
-	 * Send a game invite request
-	 */
+	/** Send a game invite request. */
 	void SendGameInviteRequest();
 
-	/**
-	 * Send join game request
-	 */
+	/** Send join game request. */
 	void SendGameJoinRequest();
 
-	/**
-	 * Send a chat message.
-	 */
+	/** Send a chat message. */
 	void SendChatMessageRequest();
 
 	/**
-	 * Delegate used when enumeration of messages is complete
+	 * Delegate used when enumeration of messages is complete.
 	 *
-	 * @param LocalPlayer		- the controller number of the associated user that made the request
-	 * @param bWasSuccessful	- true if the async action completed without error, false if there was an error
-	 * @param ErrorStr			- string representing the error condition
+	 * @param LocalPlayer The controller number of the associated user that made the request.
+	 * @param bWasSuccessful true if the async action completed without error, false if there was an error.
+	 * @param ErrorStr String representing the error condition.
 	 */
 	void OnEnumerateMessagesComplete(int32 LocalPlayer, bool bWasSuccessful, const FString& ErrorStr);
 
 	/**
-	 * Delegate used when a read message is complete
+	 * Delegate used when a read message is complete.
 	 *
-	 * @param LocalPlayer		- the controller number of the associated user that made the request
-	 * @param bWasSuccessful	- true if the async action completed without error, false if there was an error
-	 * @param MessageId			- the message ID
-	 * @param ErrorStr			- string representing the error condition
+	 * @param LocalPlayer The controller number of the associated user that made the request.
+	 * @param bWasSuccessful true if the async action completed without error, false if there was an error.
+	 * @param MessageId The message ID.
+	 * @param ErrorStr String representing the error condition.
 	 */
 	void OnReadMessageComplete(int32 LocalPlayer, bool bWasSuccessful, const FUniqueMessageId& MessageId, const FString& ErrorStr);
 
 	/**
-	 * Delegate used when an send message has completed
+	 * Delegate used when an send message has completed.
 	 *
-	 * @param LocalPlayer		- the controller number of the associated user that made the request
-	 * @param bWasSuccessful	- true if the async action completed without error, false if there was an error
-	 * @param ErrorStr			- string representing the error condition
+	 * @param LocalPlayer The controller number of the associated user that made the request.
+	 * @param bWasSuccessful true if the async action completed without error, false if there was an error.
+	 * @param ErrorStr String representing the error condition.
 	 */
 	void OnSendMessageComplete(int32 LocalPlayer, bool bWasSuccessful, const FString& ErrorStr);
 
 	/**
-	 * Delegate used when a delete message has completed
+	 * Delegate used when a delete message has completed.
 	 *
-	 * @param LocalPlayer		- the controller number of the associated user that made the request
-	 * @param bWasSuccessful	- true if the async action completed without error, false if there was an error
-	 * @param MessageId			- the message ID
-	 * @param ErrorStr			- string representing the error condition
+	 * @param LocalPlayer The controller number of the associated user that made the request.
+	 * @param bWasSuccessful true if the async action completed without error, false if there was an error.
+	 * @param MessageId The message ID.
+	 * @param ErrorStr String representing the error condition.
 	 */
 	void OnDeleteMessageComplete(int32 LocalPlayer, bool bWasSuccessful, const FUniqueMessageId& MessageId, const FString& ErrorStr);
 
 	/**
-	 * Handle an accept message accepted from a notification
+	 * Handle an accept message accepted from a notification.
 	 *
-	 * @param MessageNotification 	- The message responded to
+	 * @param MessageNotification The message responded to.
 	 */
 	FReply HandleMessageAccepted( TSharedPtr< FFriendsAndChatMessage > MessageNotification );
 
 	/**
-	 * Handle an accept message accepted from a notification
+	 * Handle an accept message accepted from a notification.
 	 *
-	 * @param MessageNotification 	- The message responded to
+	 * @param MessageNotification The message responded to.
 	 */
 	FReply HandleMessageDeclined( TSharedPtr< FFriendsAndChatMessage > MessageNotification );
 
-	/**
-	 * Resend a failed notification message.
-	 */
+	/** Resend a failed notification message. */
 	void ResendMessage();
 
 private:
@@ -280,11 +268,12 @@ private:
 	// Holds the notification delegate
 	FOnFriendsNotification* FriendsListNotificationDelegate;
 
-/* Here we have static access for the singleton
-*****************************************************************************/
 public:
+
 	static TSharedRef< FFriendsMessageManager > Get();
 	static void Shutdown();
+
 private:
+
 	static TSharedPtr< FFriendsMessageManager > SingletonInstance;
 };
diff --git a/Engine/Source/Developer/FriendsAndChat/Private/Models/FriendListItems.cpp b/Engine/Source/Developer/FriendsAndChat/Private/Models/FriendListItems.cpp
--- a/Engine/Source/Developer/FriendsAndChat/Private/Models/FriendListItems.cpp
+++ b/Engine/Source/Developer/FriendsAndChat/Private/Models/FriendListItems.cpp
@@ -1,13 +1,11 @@
 // Copyright 1998-2014 Epic Games, Inc. All Rights Reserved.
 
-/*=============================================================================
-	FriendListItems.cpp: Implements classes and struct for the FFriendsAndChatModule.
-=============================================================================*/
-
 #include "FriendsAndChatPrivatePCH.h"
 
+
 #define LOCTEXT_NAMESPACE "FriendListItems"
 
+
 // FFriendStruct implementation
 
 void FFriendStuct::AddChild( TSharedPtr< FFriendStuct > InChild )
diff --git a/Engine/Source/Developer/FriendsAndChat/Private/Models/FriendListItems.h b/Engine/Source/Developer/FriendsAndChat/Private/Models/FriendListItems.h
--- a/Engine/Source/Developer/FriendsAndChat/Private/Models/FriendListItems.h
+++ b/Engine/Source/Developer/FriendsAndChat/Private/Models/FriendListItems.h
@@ -1,22 +1,21 @@
 // Copyright 1998-2014 Epic Games, Inc. All Rights Reserved.
 
-/*=============================================================================
-	FriendListItems.h: Declares structs and enmums used in the FriendsAndChat module.
-=============================================================================*/
-
 #pragma once
 
-// Class containing the friend information - used to build the list view
+
+/**
+ * Class containing the friend information - used to build the list view.
+ */
 class FFriendStuct
 {
 public:
 
 	/**
-	 * Constructor takes the required details
+	 * Constructor takes the required details.
 	 *
-	 * @param InOnlineFriend		The online friend
-	 * @param InOnlineUser			The online user
-	 * @param InListType			The list type
+	 * @param InOnlineFriend The online friend.
+	 * @param InOnlineUser The online user.
+	 * @param InListType The list type.
 	 */
 	FFriendStuct( TSharedPtr< FOnlineFriend > InOnlineFriend, TSharedPtr< FOnlineUser > InOnlineUser, EFriendsDisplayLists::Type InListType )
 		: GroupName(TEXT(""))
@@ -24,16 +23,16 @@ public:
 		, OnlineUser( InOnlineUser )
 		, UniqueID( InOnlineUser->GetUserId() )
 		, ListType( InListType )
-	{}
+	{ }
 
 	/**
-	 * Constructor takes the required details
+	 * Constructor takes the required details.
 	 *
-	 * @param InGroupName		The group name
+	 * @param InGroupName The group name.
 	 */
 	FFriendStuct( const FString& InGroupName )
 		: GroupName( InGroupName )
-	{}
+	{ }
 
 	// Equality operator checks friend stuct
 	friend FORCEINLINE bool operator==(const FFriendStuct& X, const FFriendStuct& Y)
@@ -56,25 +55,33 @@ public:
 
 	/**
 	 * Add a child node to this friend.
-	 * @param InChild - the child node.
+	 *
+	 * @param InChild The child node.
+	 * @see GetChildList
 	 */
 	void AddChild( TSharedPtr< FFriendStuct > InChild );
 
 	/**
 	 * Get the on-line user associated with this account.
+	 *
 	 * @return The online user.
+	 * @see GetOnlineFriend
 	 */
 	const TSharedPtr< FOnlineUser > GetOnlineUser() const;
 
 	/**
 	 * Get the cached on-line Friend.
+	 *
 	 * @return The online friend.
+	 * @see GetOnlineUser, SetOnlineFriend
 	 */
 	const TSharedPtr< FOnlineFriend > GetOnlineFriend() const;
 
 	/**
 	 * Get the friend list.
+	 *
 	 * @return The array of child nodes.
+	 * @see AddChild
 	 */
 	TArray< TSharedPtr < FFriendStuct > >& GetChildList();
 
@@ -98,7 +105,9 @@ public:
 
 	/**
 	 * Set new online friend.
-	 * @param InOnlineFriend - the new online friend.
+	 *
+	 * @param InOnlineFriend The new online friend.
+	 * @see GetOnlineFriend
 	 */
 	void SetOnlineFriend( TSharedPtr< FOnlineFriend > InOnlineFriend );
 
@@ -109,47 +118,58 @@ public:
 
 	/**
 	 * Check if we have been updated.
-	 * @return true if updated
+	 *
+	 * @return true if updated.
 	 */
 	bool IsUpdated();
 
-	/**
-	 * Set if pending invitation response.
-	 */
+	/** Set if pending invitation response. */
 	void SetPendingInvite();
 
-	/**
-	 * Set if pending invitation response.
-	 */
+	/** Set if pending invitation response. */
 	void SetPendingAccept();
 
 	/**
 	 * Get the invitation status.
-	 * @return The invitation status
+	 *
+	 * @return The invitation status.
 	 */
 	EInviteStatus::Type GetInviteStatus();
 
 private:
-	// hide the default constructor
-	FFriendStuct(): bIsUpdated(true), GroupName(TEXT("")) {};
+
+	/** Hidden default constructor. */
+	FFriendStuct()
+		: bIsUpdated(true)
+		, GroupName(TEXT(""))
+	{ };
 
 private:
-	// Holds if this item has been updated
+
+	/** Holds if this item has been updated. */
 	bool bIsUpdated;
-	// Hold the child list array
+
+	/** Hold the child list array. */
 	TArray< TSharedPtr< FFriendStuct > > Children;
-	// Holds the group name
+
+	/** Holds the group name. */
 	const FString GroupName;
-	// Holds the cached online friend
+
+	/** Holds the cached online friend. */
 	TSharedPtr<FOnlineFriend> OnlineFriend;
-	// Holds the cached online user
+
+	/** Holds the cached online user. */
 	const TSharedPtr<FOnlineUser> OnlineUser;
-	// Holds the cached user id
+
+	/** Holds the cached user id. */
 	TSharedPtr< FUniqueNetId > UniqueID;
-	// Holds if this is the list type
+
+	/** Holds if this is the list type. */
 	EFriendsDisplayLists::Type ListType;
-	// Holds if we are pending an accept as friend action
+
+	/** Holds if we are pending an accept as friend action. */
 	bool bIsPendingAccepted;
-	// Holds if we are pending an invite response
+
+	/** Holds if we are pending an invite response. */
 	bool bIsPendingInvite;
 };
diff --git a/Engine/Source/Developer/FriendsAndChat/Private/UI/Widgets/SChatWindow.cpp b/Engine/Source/Developer/FriendsAndChat/Private/UI/Widgets/SChatWindow.cpp
--- a/Engine/Source/Developer/FriendsAndChat/Private/UI/Widgets/SChatWindow.cpp
+++ b/Engine/Source/Developer/FriendsAndChat/Private/UI/Widgets/SChatWindow.cpp
@@ -1,14 +1,11 @@
 // Copyright 1998-2014 Epic Games, Inc. All Rights Reserved.
 
-/*=============================================================================
-	SChatWindow.cpp: Implements the SChatWindow class.
-	Displays the chat window
-=============================================================================*/
-
 #include "FriendsAndChatPrivatePCH.h"
 
+
 #define LOCTEXT_NAMESPACE "SChatWindow"
 
+
 void SChatWindow::Construct( const FArguments& InArgs )
 {
 	// Dragable bar for the Friends list
diff --git a/Engine/Source/Developer/FriendsAndChat/Private/UI/Widgets/SChatWindow.h b/Engine/Source/Developer/FriendsAndChat/Private/UI/Widgets/SChatWindow.h
--- a/Engine/Source/Developer/FriendsAndChat/Private/UI/Widgets/SChatWindow.h
+++ b/Engine/Source/Developer/FriendsAndChat/Private/UI/Widgets/SChatWindow.h
@@ -1,11 +1,8 @@
 // Copyright 1998-2014 Epic Games, Inc. All Rights Reserved.
 
-/*=============================================================================
-	SChatWindow.h: Declares the SChatWindow class.
-=============================================================================*/
-
 #pragma once
 
+
 /**
  * Implements the FriendsList
  */
diff --git a/Engine/Source/Developer/FriendsAndChat/Private/UI/Widgets/SFriendsItem.cpp b/Engine/Source/Developer/FriendsAndChat/Private/UI/Widgets/SFriendsItem.cpp
--- a/Engine/Source/Developer/FriendsAndChat/Private/UI/Widgets/SFriendsItem.cpp
+++ b/Engine/Source/Developer/FriendsAndChat/Private/UI/Widgets/SFriendsItem.cpp
@@ -1,14 +1,11 @@
 // Copyright 1998-2014 Epic Games, Inc. All Rights Reserved.
 
-/*=============================================================================
-	SFriendsItem.cpp: Implements the SFriendsItem class.
-	Displays Friend widget in friends list
-=============================================================================*/
-
 #include "FriendsAndChatPrivatePCH.h"
 
+
 #define LOCTEXT_NAMESPACE "SFriendsItem"
 
+
 void SFriendsItem::Construct( const FArguments& InArgs, const TSharedRef<STableViewBase>& InOwnerTableView )
 {
 	STableRow< TSharedPtr<FFriendStuct> >::ConstructInternal(
diff --git a/Engine/Source/Developer/FriendsAndChat/Private/UI/Widgets/SFriendsItem.h b/Engine/Source/Developer/FriendsAndChat/Private/UI/Widgets/SFriendsItem.h
--- a/Engine/Source/Developer/FriendsAndChat/Private/UI/Widgets/SFriendsItem.h
+++ b/Engine/Source/Developer/FriendsAndChat/Private/UI/Widgets/SFriendsItem.h
@@ -1,12 +1,8 @@
 // Copyright 1998-2014 Epic Games, Inc. All Rights Reserved.
 
-/*=============================================================================
-	SFriendsItem.h: Declares the SFriendsItem class.
-	Displays a line in the friends list view 
-=============================================================================*/
-
 #pragma once
 
+
 // Enum to list friend actions
 namespace EFriendsActionType
 {
@@ -22,6 +18,7 @@ namespace EFriendsActionType
 	};
 };
 
+
 class SFriendsItem
 	: public STableRow< TSharedPtr<FFriendStuct> >
 {
@@ -35,8 +32,9 @@ public:
 
 	/**
 	 * Constructs the widget.
-	 * @param InArgs - The Slate argument list.
-	 * @param InOwnerTableView - The owning table.
+	 *
+	 * @param InArgs The Slate argument list.
+	 * @param InOwnerTableView The owning table.
 	 */
 	void Construct( const FArguments& InArgs, const TSharedRef<STableViewBase>& InOwnerTableView );
 
@@ -44,31 +42,34 @@ private:
 
 	/**
 	 * Handle a menu option selected.
-	 * @param FriendActionType - the action type.
+	 *
+	 * @param FriendActionType The action type.
 	 */	
 	FReply OnOptionSelected( EFriendsActionType::Type FriendActionType );
 
 	/**
-	* Returns the text representation of an action.
-	* @param FriendActionType - The action type.
-	* @return A text value.
-	*/
+	 * Returns the text representation of an action.
+	 *
+	 * @param FriendActionType The action type.
+	 * @return A text value.
+	 */
 	const FText& GetActionText( EFriendsActionType::Type FriendActionType );
 
 	/**
-	* Generate the menu widget.
-	* @return A menu widget.
-	*/
+	 * Generate the menu widget.
+	 *
+	 * @return A menu widget.
+	 */
 	TSharedRef<SWidget> OnGetMenuContent();
 
 private:
-	// Holds a reference to the FriendItem data that is displayed in this row.
+
+	/** Holds a reference to the FriendItem data that is displayed in this row. */
 	TSharedPtr< FFriendStuct > FriendItem;
 
-	// Holds the style to use when making the widget.
+	/** Holds the style to use when making the widget. */
 	FFriendsAndChatStyle FriendStyle;
 
-	// Holds the menu button.
+	/** Holds the menu button. */
 	TSharedPtr< SComboButton > MenuButton; 
 };
-
diff --git a/Engine/Source/Developer/FriendsAndChat/Private/UI/Widgets/SFriendsList.cpp b/Engine/Source/Developer/FriendsAndChat/Private/UI/Widgets/SFriendsList.cpp
--- a/Engine/Source/Developer/FriendsAndChat/Private/UI/Widgets/SFriendsList.cpp
+++ b/Engine/Source/Developer/FriendsAndChat/Private/UI/Widgets/SFriendsList.cpp
@@ -1,14 +1,11 @@
 // Copyright 1998-2014 Epic Games, Inc. All Rights Reserved.
 
-/*=============================================================================
-	SFriendsList.cpp: Implements the SFriendsList class.
-	Displays the list of friends
-=============================================================================*/
-
 #include "FriendsAndChatPrivatePCH.h"
 
+
 #define LOCTEXT_NAMESPACE "SFriendsList"
 
+
 void SFriendsList::Construct( const FArguments& InArgs )
 {
 	// Dragable bar for the Friends list
diff --git a/Engine/Source/Developer/FriendsAndChat/Private/UI/Widgets/SFriendsList.h b/Engine/Source/Developer/FriendsAndChat/Private/UI/Widgets/SFriendsList.h
--- a/Engine/Source/Developer/FriendsAndChat/Private/UI/Widgets/SFriendsList.h
+++ b/Engine/Source/Developer/FriendsAndChat/Private/UI/Widgets/SFriendsList.h
@@ -1,11 +1,8 @@
 // Copyright 1998-2014 Epic Games, Inc. All Rights Reserved.
 
-/*=============================================================================
-	SFriendsList.h: Declares the SFriendsList class.
-=============================================================================*/
-
 #pragma once
 
+
 /**
  * Implements the FriendsList
  */
@@ -25,7 +22,7 @@ public:
 	/**
 	 * Constructs the application.
 	 *
-	 * @param InArgs - The Slate argument list.
+	 * @param InArgs The Slate argument list.
 	 */
 	void Construct( const FArguments& InArgs );
 
@@ -33,112 +30,120 @@ private:
 
 	/**
 	 * Generates a widget for a Friend item.
-	 * @param InItem - the FriendItem
-	 * @param OwnerTable - the owning table
-	 * @return The table row widget
+	 *
+	 * @param InItem The FriendItem.
+	 * @param OwnerTable The owning table.
+	 * @return The table row widget.
 	 */
 	TSharedRef<ITableRow> HandleGenerateFriendWidget( TSharedPtr< FFriendStuct > InItem, const TSharedRef<STableViewBase>& OwnerTable );
 
-	/**
-	 * Callback for when the friends list is updated - refreshes the list.
-	 */
+	/** Callback for when the friends list is updated - refreshes the list. */
 	void RefreshFriendsList();
 
 	/**
 	 * Get the list count as text.
+	 *
 	 * @return The number of friends
 	 */
 	FText GetListCountText() const;
 
-	/**
-	 * Handles the search for friend button clicked.
-	 */
+	/** Handles the search for friend button clicked. */
 	FReply SearchFriend_OnClicked();
 
 	/** 
-	 * Handle friend search committed 
-	 * @param CommentText - the entered text
-	 * @param CommitInfo - the commit type
+	 * Handle friend search committed.
+	 *
+	 * @param CommentText The entered text.
+	 * @param CommitInfo The commit type.
 	 */
 	void HandleFriendEntered(const FText& CommentText, ETextCommit::Type CommitInfo);
 
-	/**
-	 * Handles the close button clicked.
-	 */
+	/** Handles the close button clicked. */
 	FReply CloseButton_OnClicked();
 
-	/**
-	 * Handles the minimize button clicked.
-	 */
+	/** Handles the minimize button clicked. */
 	FReply MinimizeButton_OnClicked();
 
 	/**
 	 * Handles the default list selection button change.
-	 * @param Check State - The checkstate
+	 *
+	 * @param CheckState The new check state.
 	 */
 	void DefaultListSelect_OnChanged( ESlateCheckBoxState::Type CheckState );
 
 	/**
 	 * Get the check state for the default list selection check box.
-	 * @return The checkstate
+	 *
+	 * @return The check state.
 	 */
 	ESlateCheckBoxState::Type DefaultListSelect_IsChecked() const;
 
 	/**
 	 * Handles the recent players list selection button change.
-	 * @param Check State - The checkstate
+	 *
+	 * @param CheckState The new check state.
 	 */
 	void RecentListSelect_OnChanged( ESlateCheckBoxState::Type CheckState );
 
 	/**
 	 * Handles the player requests list selection button change.
-	 * @param Check State - The checkstate.
+	 * @param CheckState The check state.
 	 */
 	void RequestListSelect_OnChanged( ESlateCheckBoxState::Type CheckState );
 
 	/**
 	 * Get the check state for the recent players list selection check box.
-	 * @return The checkstate
+	 *
+	 * @return The check state.
 	 */
 	ESlateCheckBoxState::Type RecentListSelect_IsChecked();
 
-	/**
-	 * Set the toggle button state depending on the selected list.
-	 */
+	/** Set the toggle button state depending on the selected list. */
 	void UpdateButtonState();
 
 	/**
 	 * Get the menu text color depending on the selected list.
-	 * @param ListType - the list to get the color for
-	 * @return The list color
+	 *
+	 * @param ListType The list to get the color for.
+	 * @return The list color.
 	 */
 	FSlateColor GetMenuTextColor( EFriendsDisplayLists::Type ListType ) const;
 
 private:
 
-	// Holds the list of friends
+	/** Holds the list of friends. */
 	TArray< TSharedPtr< FFriendStuct > > FriendsList;
-	// Holds the list of outgoing invites
+	
+	/** Holds the list of outgoing invites. */
 	TArray< TSharedPtr< FFriendStuct > > OutgoingFriendsList;
-	// Holds the tree view of the Friends list
+	
+	/** Holds the tree view of the Friends list. */
 	TSharedPtr< SListView< TSharedPtr< FFriendStuct > > > FriendsListView;
-	// Holds the style to use when making the widget
+	
+	/** Holds the style to use when making the widget. */
 	FFriendsAndChatStyle FriendStyle;
-	// Holds the delegate for when the close button is clicked
+	
+	/** Holds the delegate for when the close button is clicked. */
 	FOnClicked OnCloseClicked;
-	// Holds the delegate for when the minimize button is clicked
+	
+	/** Holds the delegate for when the minimize button is clicked. */
 	FOnClicked OnMinimizeClicked;
-	// Holds the text box used to enter the name of friend to search for
+	
+	/** Holds the text box used to enter the name of friend to search for. */
 	TSharedPtr< SEditableTextBox > FriendNameTextBox;
-	// Holds the recent players check box
+	
+	/** Holds the recent players check box. */
 	TSharedPtr< SCheckBox > RecentPlayersButton;
-	// Holds the recent players check box
+	
+	/** Holds the recent players check box. */
 	TSharedPtr< SCheckBox > FriendRequestsButton;
-	// Holds the default friends check box
+	
+	/** Holds the default friends check box. */
 	TSharedPtr< SCheckBox > DefaultPlayersButton;
-	// Holds the default list name
+	
+	/** Holds the default list name. */
 	EFriendsDisplayLists::Type CurrentList;
-	// Holds the friends list display box.
+	
+	/** Holds the friends list display box. */
 	TSharedPtr< SVerticalBox > DisplayBox;
 };
-
diff --git a/Engine/Source/Developer/FriendsAndChat/Private/UI/Widgets/SFriendsListContextMenu.cpp b/Engine/Source/Developer/FriendsAndChat/Private/UI/Widgets/SFriendsListContextMenu.cpp
--- a/Engine/Source/Developer/FriendsAndChat/Private/UI/Widgets/SFriendsListContextMenu.cpp
+++ b/Engine/Source/Developer/FriendsAndChat/Private/UI/Widgets/SFriendsListContextMenu.cpp
@@ -1,13 +1,11 @@
 // Copyright 1998-2014 Epic Games, Inc. All Rights Reserved.
 
-/*=============================================================================
-	SFriendsListContextMenu.cpp: Implements the SFriendsListContextMenu class.
-=============================================================================*/
-
 #include "FriendsAndChatPrivatePCH.h"
 
+
 #define LOCTEXT_NAMESPACE "SFriendsListContextMenu"
 
+
 void SFriendsListContextMenu::Construct( const FArguments& InArgs )
 {
 	FriendItem = InArgs._FriendItem;
@@ -18,6 +16,7 @@ void SFriendsListContextMenu::Construct( const FArguments& InArgs )
 	];
 }
 
+
 TSharedRef<SWidget> SFriendsListContextMenu::MakeContextMenu( )
 {
 	static const FText MenuAddFriend( LOCTEXT( "FriendListContextMenu_AddFriend", "ADD FRIEND" ) );
@@ -78,6 +77,7 @@ TSharedRef<SWidget> SFriendsListContextMenu::MakeContextMenu( )
 	return MenuBuilder.MakeWidget();
 }
 
+
 void SFriendsListContextMenu::OnAddFriendClicked()
 {
 	if ( FriendItem.IsValid() )
@@ -86,6 +86,7 @@ void SFriendsListContextMenu::OnAddFriendClicked()
 	}
 }
 
+
 void SFriendsListContextMenu::OnChatClicked()
 {
 	if ( FriendItem.IsValid() )
@@ -94,6 +95,7 @@ void SFriendsListContextMenu::OnChatClicked()
 	}
 }
 
+
 void SFriendsListContextMenu::OnRemoveFriendClicked()
 {
 	if ( FriendItem.IsValid() && FriendItem->GetOnlineFriend().IsValid() )
@@ -102,6 +104,7 @@ void SFriendsListContextMenu::OnRemoveFriendClicked()
 	}
 }
 
+
 void SFriendsListContextMenu::OnInviteToGameClicked()
 {
 	if ( FriendItem.IsValid() && FriendItem->GetOnlineFriend().IsValid() )
@@ -110,6 +113,7 @@ void SFriendsListContextMenu::OnInviteToGameClicked()
 	}
 }
 
+
 void SFriendsListContextMenu::OnJoinGameClicked()
 {
 	if ( FriendItem.IsValid() && FriendItem->GetOnlineFriend().IsValid() )
@@ -118,6 +122,7 @@ void SFriendsListContextMenu::OnJoinGameClicked()
 	}
 }
 
+
 void SFriendsListContextMenu::OnAcceptFriendClicked()
 {
 	if ( FriendItem.IsValid() && FriendItem->GetOnlineUser().IsValid() )
@@ -126,6 +131,7 @@ void SFriendsListContextMenu::OnAcceptFriendClicked()
 	}
 }
 
+
 void SFriendsListContextMenu::OnRejectFriendClicked()
 {
 	if ( FriendItem.IsValid() && FriendItem->GetOnlineUser().IsValid() )
@@ -133,4 +139,6 @@ void SFriendsListContextMenu::OnRejectFriendClicked()
 		FFriendsAndChatManager::Get()->RejectFriend( FriendItem );
 	}
 }
+
+
 #undef LOCTEXT_NAMESPACE
diff --git a/Engine/Source/Developer/FriendsAndChat/Private/UI/Widgets/SFriendsListContextMenu.h b/Engine/Source/Developer/FriendsAndChat/Private/UI/Widgets/SFriendsListContextMenu.h
--- a/Engine/Source/Developer/FriendsAndChat/Private/UI/Widgets/SFriendsListContextMenu.h
+++ b/Engine/Source/Developer/FriendsAndChat/Private/UI/Widgets/SFriendsListContextMenu.h
@@ -1,11 +1,8 @@
 // Copyright 1998-2014 Epic Games, Inc. All Rights Reserved.
 
-/*=============================================================================
-	SFriendsListContextMenu.h: Declares the SFriendsListContextMenu class.
-=============================================================================*/
-
 #pragma once
 
+
 class SFriendsListContextMenu
 	: public SCompoundWidget
 {
@@ -19,7 +16,7 @@ class SFriendsListContextMenu
 	/**
 	 * Constructs the application.
 	 *
-	 * @param InArgs - The Slate argument list.
+	 * @param InArgs The Slate argument list.
 	 */
 	void Construct( const FArguments& InArgs );
 
@@ -27,47 +24,34 @@ private:
 
 	/**
 	 * Make the context menu.
+	 *
 	 * @return The context menu widget
 	 */
 	TSharedRef<SWidget> MakeContextMenu();
 
-	/**
-	 * Handle the Add Friend button clicked.
-	 */
+	/** Handle the Add Friend button clicked. */
 	void OnAddFriendClicked();
 
-	/**
-	 * Handle the chat clicked.
-	 */
+	/** Handle the chat clicked. */
 	void OnChatClicked();
 
-	/**
-	 * Handle the remove friend clicked.
-	 */
+	/** Handle the remove friend clicked. */
 	void OnRemoveFriendClicked();
 
-	/**
-	 * Handle the remove friend clicked.
-	 */
+	/** Handle the remove friend clicked. */
 	void OnInviteToGameClicked();
 
-	/**
-	 * Handle join game clicked.
-	 */
+	/** Handle join game clicked. */
 	void OnJoinGameClicked();
 
-	/**
-	 * Handle accept friend clicked.
-	 */
+	/** Handle accept friend clicked. */
 	void OnAcceptFriendClicked();
 
-	/**
-	 * Handle reject friend clicked.
-	 */
+	/** Handle reject friend clicked. */
 	void OnRejectFriendClicked();
 
 public:
-	// Holds a reference to the news message that is displayed in this row.
+
+	/** Holds a reference to the news message that is displayed in this row. */
 	TSharedPtr< FFriendStuct > FriendItem;
 };
-
diff --git a/Engine/Source/Developer/FriendsAndChat/Public/FriendsAndChat.h b/Engine/Source/Developer/FriendsAndChat/Public/FriendsAndChat.h
--- a/Engine/Source/Developer/FriendsAndChat/Public/FriendsAndChat.h
+++ b/Engine/Source/Developer/FriendsAndChat/Public/FriendsAndChat.h
@@ -1,13 +1,11 @@
 // Copyright 1998-2014 Epic Games, Inc. All Rights Reserved.
 
-/*=============================================================================
-	FriendsAndChat.h: FriendsAndChat module public header file.
-=============================================================================*/
-
 #pragma once
 
+
 /* Dependencies
  *****************************************************************************/
+
 #include "Core.h"
 #include "ModuleManager.h"
 #include "Http.h"
@@ -16,8 +14,10 @@
 #include "OnlineSubsystem.h"
 #include "OnlineSubsystemMcp.h"
 
+
 /* Interfaces
  *****************************************************************************/
+
 #include "FriendsAndChatMessage.h"
 #include "IFriendsAndChatModule.h"
 #include "FriendsAndChatStyle.h"
diff --git a/Engine/Source/Developer/FriendsAndChat/Public/Interfaces/IFriendsAndChatModule.h b/Engine/Source/Developer/FriendsAndChat/Public/Interfaces/IFriendsAndChatModule.h
--- a/Engine/Source/Developer/FriendsAndChat/Public/Interfaces/IFriendsAndChatModule.h
+++ b/Engine/Source/Developer/FriendsAndChat/Public/Interfaces/IFriendsAndChatModule.h
@@ -1,17 +1,13 @@
 // Copyright 1998-2014 Epic Games, Inc. All Rights Reserved.
 
-/*=============================================================================
-	IFriendsAndChatModule.h: Declares the IFriendsAndChatModule interface.
-=============================================================================*/
-
 #pragma once
 
-/**
- * Delegate type for FriendsList updated.
- */
+
+/** Delegate type for FriendsList updated. */
 DECLARE_MULTICAST_DELEGATE_OneParam( FOnFriendsNotification, const FFriendsAndChatMessageRef& )
 DECLARE_DELEGATE_RetVal(bool, FFriendsSystemReady )
 
+
 /**
  * Interface for the Friends and chat manager.
  */
@@ -47,62 +43,62 @@ public:
 
 	/**
 	 * Initialization function.
-	 * @param bAllowLaunchGame - Can this app launch a game directly. CP cannot
+	 *
+	 * @param bAllowLaunchGame Can this app launch a game directly. CP cannot
 	 */
 	virtual void Init( bool bAllowLaunchGame = true ) = 0;
 
 	/**
 	 * Returns a delegate that is executed when we receive a Friends List notification.
+	 *
 	 * @return The delegate.
 	 */
 	 virtual FOnFriendsNotification& OnFriendsListNotification() = 0;
 
 	/**
 	 * Create the friends list widget.
-	 * @param ParentWindow - the parent window to add the widget
-	 * @param InStyle - the style to use to create the widgets
+	 *
+	 * @param ParentWindow The parent window to add the widget
+	 * @param InStyle The style to use to create the widgets
 	 */
 	virtual void CreateFriendsListWidget( TSharedPtr< const SWidget > ParentWidget, const struct FFriendsAndChatStyle* InStyle ) = 0;
 
 	/**
 	 * Create the a friends list widget without a container.
-	 * @param InStyle - the style to use to create the widgets
-	 * @return The Friends List widget
+	 * @param InStyle The style to use to create the widgets.
+	 * @return The Friends List widget.
 	 */
 	virtual TSharedPtr< SWidget > GenerateFriendsListWidget( const struct FFriendsAndChatStyle* InStyle ) = 0;
 
 	/**
 	 * Set that we are in a session, so can send join game requests.
-	 * @param bInSession - If we are in a session
+	 *
+	 * @param bInSession If we are in a session.
 	 */
 	virtual void SetInSession( bool bInSession ) = 0;
 
 	/**
 	 * Callback for if a notification failed. Will resend later.
-	 * @param NetID - the original caller ID
+	 *
+	 * @param NetID The original caller ID.
 	 */
 	virtual void SetUnhandledNotification( TSharedRef< FUniqueNetId > NetID ) = 0;
 
-	/**
-	 * Clear game invites.
-	 */
+	/** Clear game invites. */
 	virtual void ClearGameInvites() = 0;
 
 	/**
 	 * Get the friends list count.
-	 * @return The friends list count
+	 *
+	 * @return The friends list count.
 	 */
 	virtual int32 GetFriendCount() = 0;
 
-	/**
-	 * Log out - kill the friends list window.
-	 */
+	/** Log out and kill the friends list window. */
 	virtual void Logout() = 0;
 
 public:
 
-	/**
-	 * Virtual destructor.
-	 */
-	virtual ~IFriendsAndChatModule() {}
+	/** Virtual destructor. */
+	virtual ~IFriendsAndChatModule() { }
 };
diff --git a/Engine/Source/Developer/FriendsAndChat/Public/Models/FriendsAndChatMessage.h b/Engine/Source/Developer/FriendsAndChat/Public/Models/FriendsAndChatMessage.h
--- a/Engine/Source/Developer/FriendsAndChat/Public/Models/FriendsAndChatMessage.h
+++ b/Engine/Source/Developer/FriendsAndChat/Public/Models/FriendsAndChatMessage.h
@@ -1,18 +1,17 @@
 // Copyright 1998-2014 Epic Games, Inc. All Rights Reserved.
 
-/*=============================================================================
-	FriendAndChatMessage.h: Declares the FriendsAndChatMessage class.
-=============================================================================*/
-
 #pragma once
 
+
 class FUniqueNetId;
 
+
 /**
  * Type definition for shared references to instances of ILauncherProfile.
  */
 typedef TSharedRef<class FFriendsAndChatMessage> FFriendsAndChatMessageRef;
 
+
 // Enum to list download status flags
 namespace EFriendsRequestType
 {
@@ -31,6 +30,7 @@ namespace EFriendsRequestType
 	};
 };
 
+
 // Enum to list friend respose type
 namespace EFriendsResponseType
 {
@@ -45,6 +45,7 @@ namespace EFriendsResponseType
 	};
 };
 
+
 // Enum holding the display list selection
 namespace EFriendsDisplayLists
 {
@@ -56,32 +57,36 @@ namespace EFriendsDisplayLists
 	};
 };
 
+
 // Class containing the friend message information
 class FFriendsAndChatMessage
 {
 public:
 	/**
 	 * Constructor.
-	 * @param InMessage - The message content.
+	 *
+	 * @param InMessage The message content.
 	 */
 	FFriendsAndChatMessage( const FString& InMessage )
 		: MessageConent( InMessage )
 		, bAutoAccept( false )
-	{}
+	{ }
 
 	/**
 	 * Constructor.
-	 * @param InMessage			- The message content.
-	 * @param InUniqueFriendID	- The Friend ID.
+	 *
+	 * @param InMessage The message content.
+	 * @param InUniqueFriendID The Friend ID.
 	 */
 	FFriendsAndChatMessage( const FString& InMessage, const TSharedRef< FUniqueNetId > InUniqueFriendID )
 		: MessageConent( InMessage )
 		, UniqueFriendID( InUniqueFriendID )
 		, bAutoAccept( false )
-	{}
+	{ }
 
 	/**
 	 * Get the message content.
+	 *
 	 * @return The message.
 	 */
 	FString GetMessage()
@@ -91,6 +96,7 @@ public:
 
 	/**
 	 * Get the Net ID for the Friend this message is sent to.
+	 *
 	 * @return The net ID.
 	 */
 	const TSharedRef< FUniqueNetId > GetUniqueID()
@@ -100,7 +106,8 @@ public:
 
 	/**
 	 * Set a button callbacks.
-	 * @param InCallback - The button callback.
+	 *
+	 * @param InCallback The button callback.
 	 */
 	void SetButtonCallback( FOnClicked InCallback )
 	{
@@ -109,7 +116,8 @@ public:
 
 	/**
 	 * Set the message type.
-	 * @param InType - The message type.
+	 *
+	 * @param InType The message type.
 	 */
 	void SetMessageType( EFriendsRequestType::Type InType )
 	{
@@ -117,8 +125,9 @@ public:
 	}
 
 	/**
-	 * Sets the ID to use when joinging a game from an invite. Passed through the commandline.
-	 * @param MessageID - The string representation of the NetID.
+	 * Sets the ID to use when joining a game from an invite. Passed through the command line.
+	 *
+	 * @param MessageID The string representation of the NetID.
 	 */
 	void SetLaunchGameID( const FString& MessageID )
 	{
@@ -126,7 +135,8 @@ public:
 	}
 	
 	/**
-	 * Gets the ID to use when joinging a game from an invite. Passed through the commandline.
+	 * Gets the ID to use when joining a game from an invite. Passed through the command line.
+	 *
 	 * @return string representation of the message ID.
 	 */
 	const FString& GetLauchGameID() const
@@ -136,6 +146,7 @@ public:
 
 	/**
 	 * Get the message type.
+	 *
 	 * @return The message type.
 	 */
 	EFriendsRequestType::Type GetMessageType()
@@ -145,7 +156,8 @@ public:
 
 	/**
 	 * Set the requester name.
-	 * @param InName - The requesters name.
+	 *
+	 * @param InName The requesters name.
 	 */
 	void SetRequesterName( const FString& InName )
 	{
@@ -154,7 +166,8 @@ public:
 
 	/**
 	 * Get the requester name.
-	 * @return - The requesters name.
+	 *
+	 * @return The requesters name.
 	 */
 	const FString& GetRequesterName() const
 	{
@@ -163,16 +176,15 @@ public:
 
 	/**
 	 * Set if this message should be handled by the caller, not the Friends Module.
-	 * @param bInSelfHandle - True if the caller handles this message themselves.
+	 *
+	 * @param bInSelfHandle True if the caller handles this message themselves.
 	 */
 	void SetSelfHandle( bool bInSelfHandle )
 	{
 		bSelfHandle = bInSelfHandle;
 	}
 
-	/**
-	 * Set if this message should accepted automatically.
-	 */
+	/** Set if this message should accepted automatically. */
 	void SetAutoAccept()
 	{
 		bAutoAccept = true;
@@ -180,6 +192,7 @@ public:
 
 	/**
 	 * Should the caller auto accept this message.
+	 *
 	 * @return True if the user should handle auto accept this message.
 	 */
 	const bool IsAutoAccepted() const
@@ -189,6 +202,7 @@ public:
 
 	/**
 	 * Should the caller handle this message themselves.
+	 *
 	 * @return True if the user should handle the message themselves.
 	 */
 	const bool IsSelfHandle() const
@@ -198,6 +212,7 @@ public:
 
 	/**
 	 * Get a button callbacks.
+	 *
 	 * @return The button callback array.
 	 */
 	TArray< FOnClicked > GetCallbacks()
@@ -205,29 +220,35 @@ public:
 		return ButtonCallbacks;
 	}
 
-	/**
-	 * Set this message into a handled state.
-	 */
+	/** Set this message into a handled state. */
 	void SetHandled()
 	{
 		MessageType = EFriendsRequestType::MessageHandled;
 	}
 
 private:
+
 	// Holds the button callbacks
 	TArray< FOnClicked > ButtonCallbacks;
+	
 	// Holds the message content
 	FString MessageConent;
+	
 	// Holds the Unique Friend Net ID
 	const TSharedPtr< FUniqueNetId > UniqueFriendID;
+	
 	// Holds the message type
 	EFriendsRequestType::Type MessageType;
+	
 	// Holds the requester name.
 	FString RequesterName;
+	
 	// Holds if the message should be self handled.
 	bool bSelfHandle;
+	
 	// Holds if this message should be auto accepted.
 	bool bAutoAccept;
+	
 	// Holds the launch game ID
 	FString LaunchGameID;
 };
diff --git a/Engine/Source/Developer/LauncherServices/Public/Interfaces/ILauncher.h b/Engine/Source/Developer/LauncherServices/Public/Interfaces/ILauncher.h
--- a/Engine/Source/Developer/LauncherServices/Public/Interfaces/ILauncher.h
+++ b/Engine/Source/Developer/LauncherServices/Public/Interfaces/ILauncher.h
@@ -20,17 +20,14 @@ public:
 	/**
 	 * Launches the specified profile.
 	 *
-	 * @param DeviceProxyManager - The target device proxy manager to use.
-	 * @param Profile - The profile to launch.
-	 *
-	 * @return The worker thread, or NULL if not launched.
+	 * @param DeviceProxyManager The target device proxy manager to use.
+	 * @param Profile The profile to launch.
+	 * @return The worker thread, or nullptr if not launched.
 	 */
 	virtual ILauncherWorkerPtr Launch( const ITargetDeviceProxyManagerRef& DeviceProxyManager, const ILauncherProfileRef& Profile ) = 0;
 
 public:
 
-	/**
-	 * Virtual destructor.
-	 */
+	/** Virtual destructor. */
 	virtual ~ILauncher( ) { }
 };
diff --git a/Engine/Source/Developer/LauncherServices/Public/Interfaces/ILauncherDeviceGroup.h b/Engine/Source/Developer/LauncherServices/Public/Interfaces/ILauncherDeviceGroup.h
--- a/Engine/Source/Developer/LauncherServices/Public/Interfaces/ILauncherDeviceGroup.h
+++ b/Engine/Source/Developer/LauncherServices/Public/Interfaces/ILauncherDeviceGroup.h
@@ -37,7 +37,7 @@ public:
 	/**
 	 * Adds a device to the group.
 	 *
-	 * @param DeviceID - The identifier of the device to add.
+	 * @param DeviceID The identifier of the device to add.
 	 */
 	virtual void AddDevice(const FString& DeviceID) = 0;
 	
@@ -72,14 +72,14 @@ public:
 	/**
 	 * Removes a device from the group.
 	 *
-	 * @param DeviceID - The identifier of the device to remove.
+	 * @param DeviceID The identifier of the device to remove.
 	 */
 	virtual void RemoveDevice(const FString& DeviceID) = 0;
 
 	/**
 	 * Sets the name of the device group.
 	 *
-	 * @param NewName - The new name.
+	 * @param NewName The new name.
 	 */
 	virtual void SetName(const FString& NewName) = 0;
 
@@ -101,8 +101,6 @@ public:
 
 public:
 
-	/**
-	 * Virtual destructor.
-	 */
+	/** Virtual destructor. */
 	virtual ~ILauncherDeviceGroup() { }
 };
diff --git a/Engine/Source/Developer/LauncherServices/Public/Interfaces/ILauncherProfile.h b/Engine/Source/Developer/LauncherServices/Public/Interfaces/ILauncherProfile.h
--- a/Engine/Source/Developer/LauncherServices/Public/Interfaces/ILauncherProfile.h
+++ b/Engine/Source/Developer/LauncherServices/Public/Interfaces/ILauncherProfile.h
@@ -181,7 +181,6 @@ public:
 	 * Gets the device variant to use when deploying and launching.
 	 *
 	 * @return Device Variant name.
-	 *
 	 * @see SetDeviceVariant
 	 */
 	virtual FName GetDeviceVariant() const = 0;
@@ -190,7 +189,6 @@ public:
 	 * Gets the name of the build configuration.
 	 *
 	 * @return Build configuration name.
-	 *
 	 * @see SetBuildConfigurationName
 	 */
 	virtual EBuildConfigurations::Type GetBuildConfiguration() const = 0;
@@ -205,15 +203,14 @@ public:
 	/**
 	 * Updates the device name.
 	 *
-	 * @param InDeviceName - The new device name.
+	 * @param InDeviceName The new device name.
 	 */
 	virtual void SetDeviceName(const FString& InDeviceName) = 0;
 
 	/**
 	 * Sets the device variant.
 	 *
-	 * @param InVariant - The variant to set.
-	 *
+	 * @param InVariant The variant to set.
 	 * @see GetDeviceVariant
 	 */
 	virtual void SetDeviceVariant(FName InVariant) = 0;
@@ -221,8 +218,7 @@ public:
 	/**
 	 * Sets the build configuration.
 	 *
-	 * @param InConfiguration - The build configuration name to set.
-	 *
+	 * @param InConfiguration The build configuration name to set.
 	 * @see GetBuildConfigurationName
 	 */
 	virtual void SetBuildConfiguration(EBuildConfigurations::Type InConfiguration) = 0;
@@ -230,8 +226,7 @@ public:
 	/**
 	 * Sets the cook mode.
 	 *
-	 * @param InMode - The cook mode.
-	 *
+	 * @param InMode The cook mode.
 	 * @see GetCookMode
 	 */
 	virtual void SetCookMode(ELauncherProfileCookModes::Type InMode) = 0;
@@ -239,18 +234,14 @@ public:
 	/**
 	 * Serializes the simple profile from or into the specified archive.
 	 *
-	 * @param Archive - The archive to serialize from or into.
-	 *
+	 * @param Archive The archive to serialize from or into.
 	 * @return true if the profile was serialized, false otherwise.
 	 */
 	virtual bool Serialize(FArchive& Archive) = 0;
 
-	/**
-	 * Sets all profile settings to their defaults.
-	 */
+	/** Sets all profile settings to their defaults. */
 	virtual void SetDefaults() = 0;
 
-
 public:
 
 	/**
@@ -266,6 +257,7 @@ typedef TSharedPtr<class ILauncherProfile> ILauncherProfilePtr;
 /** Type definition for shared references to instances of ILauncherProfile. */
 typedef TSharedRef<class ILauncherProfile> ILauncherProfileRef;
 
+
 /**
  * Delegate type for changing the device group to deploy to.
  *
@@ -281,6 +273,7 @@ DECLARE_MULTICAST_DELEGATE(FOnProfileProjectChanged);
  */
 DECLARE_DELEGATE_RetVal(bool, FIsCookFinishedDelegate);
 
+
 /**
  * Interface for launcher profile.
  */
@@ -319,8 +312,7 @@ public:
 	/**
 	 * Checks whether the last validation yielded the specified error.
 	 *
-	 * @param Error - The validation error to check for.
-	 *
+	 * @param Error The validation error to check for.
 	 * @return true if the error is present, false otherwise.
 	 */
 	virtual bool HasValidationError( ELauncherProfileValidationErrors::Type Error ) const = 0;
@@ -338,8 +330,7 @@ public:
 	 * Whether a platform is deployable depends on the current profile settings.
 	 * The right combination of build, cook and package settings must be present.
 	 *
-	 * @param PlatformName - The name of the platform to deploy.
-	 *
+	 * @param PlatformName The name of the platform to deploy.
 	 * @return true if the platform is deployable, false otherwise.
 	 */
 	virtual bool IsDeployablePlatform( const FString& PlatformName ) = 0;
@@ -354,43 +345,41 @@ public:
 	/**
 	 * Serializes the profile from or into the specified archive.
 	 *
-	 * @param Archive - The archive to serialize from or into.
-	 *
+	 * @param Archive The archive to serialize from or into.
 	 * @return true if the profile was serialized, false otherwise.
 	 */
 	virtual bool Serialize( FArchive& Archive ) = 0;
 
-	/**
-	 * Sets all profile settings to their defaults.
-	 */
+	/** Sets all profile settings to their defaults. */
 	virtual void SetDefaults( ) = 0;
 
 	/**
 	 * Updates the name of the profile.
 	 *
-	 * @param NewName - The new name of the profile.
+	 * @param NewName The new name of the profile.
 	 */
 	virtual void SetName( const FString& NewName ) = 0;
 
 	/**
 	 * Updates the description of the profile.
 	 *
-	 * @param NewDescription - The new description of the profile.
+	 * @param NewDescription The new description of the profile.
 	 */
 	virtual void SetDescription(const FString& NewDescription) = 0;
 
 	/**
-	 * Returns the cook delegate which can be used to query if the cook is finished
-	 *  Used by cook by the book in the editor
+	 * Returns the cook delegate which can be used to query if the cook is finished.
+	 *
+	 * Used by cook by the book in the editor.
 	 */
 	virtual FIsCookFinishedDelegate& OnIsCookFinished() = 0;
+
 public:
 
 	/**
 	 * Gets the name of the build configuration.
 	 *
 	 * @return Build configuration name.
-	 *
 	 * @see SetBuildConfigurationName
 	 */
 	virtual EBuildConfigurations::Type GetBuildConfiguration( ) const = 0;
@@ -399,7 +388,6 @@ public:
 	 * Gets the build configuration name of the cooker.
 	 *
 	 * @return Cook configuration name.
-	 *
 	 * @see SetCookConfigurationName
 	 */
 	virtual EBuildConfigurations::Type GetCookConfiguration( ) const = 0;
@@ -422,10 +410,7 @@ public:
 	 * Gets the list of cooked culture.
 	 *
 	 * @return Collection of culture names.
-	 *
-	 * @see AddCookedCulture
-	 * @see ClearCookedCultures
-	 * @see RemoveCookedCulture
+	 * @see AddCookedCulture, ClearCookedCultures, RemoveCookedCulture
 	 */
 	virtual const TArray<FString>& GetCookedCultures( ) const = 0;
 
@@ -434,9 +419,7 @@ public:
 	 *
 	 * @return Collection of map names.
 	 *
-	 * @see AddCookedMap
-	 * @see ClearCookedMaps
-	 * @see RemoveCookedMap
+	 * @see AddCookedMap, ClearCookedMaps, RemoveCookedMap
 	 */
 	virtual const TArray<FString>& GetCookedMaps( ) const = 0;
 
@@ -445,9 +428,7 @@ public:
 	 *
 	 * @return Read-only collection of platform names.
 	 *
-	 * @see AddCookedPlatform
-	 * @see ClearCookedPlatforms
-	 * @see RemoveCookedPlatform
+	 * @see AddCookedPlatform, ClearCookedPlatforms, RemoveCookedPlatform
 	 */
 	virtual const TArray<FString>& GetCookedPlatforms( ) const = 0;
 
@@ -462,7 +443,6 @@ public:
 	 * Gets the device group to deploy to.
 	 *
 	 * @return The device group, or NULL if none was configured.
-	 *
 	 * @see SetDeployedDeviceGroup
 	 */
 	virtual ILauncherDeviceGroupPtr GetDeployedDeviceGroup( ) const = 0;
@@ -471,7 +451,6 @@ public:
 	 * Gets the deployment mode.
 	 *
 	 * @return The deployment mode.
-	 *
 	 * @see SetDeploymentMode
 	 */
 	virtual ELauncherProfileDeploymentModes::Type GetDeploymentMode( ) const = 0;
@@ -479,8 +458,7 @@ public:
     /**
      * Gets the close mode for the cook on the fly server
      *
-     * @return the close mode
-     *
+     * @return The close mode.
      * @see SetForceClose
      */
     virtual bool GetForceClose() const = 0;
@@ -489,7 +467,6 @@ public:
 	 * Gets the launch mode.
 	 *
 	 * @return The launch mode.
-	 *
 	 * @see SetLaunchMode
 	 */
 	virtual ELauncherProfileLaunchModes::Type GetLaunchMode( ) const = 0;
@@ -498,17 +475,15 @@ public:
 	 * Gets the profile's collection of launch roles.
 	 *
 	 * @return A read-only collection of launch roles.
-	 *
-	 * @see CreateLaunchRole
-	 * @see RemoveLaunchRole
+	 * @see CreateLaunchRole, RemoveLaunchRole
 	 */
 	virtual const TArray<ILauncherProfileLaunchRolePtr>& GetLaunchRoles( ) const = 0;
 
 	/**
 	 * Gets the launch roles assigned to the specified device.
 	 *
-	 * @param DeviceId - The identifier of the device.
-	 * @param OutRoles - Will hold the assigned roles, if any.
+	 * @param DeviceId The identifier of the device.
+	 * @param OutRoles Will hold the assigned roles, if any.
 	 */
 	virtual const int32 GetLaunchRolesFor( const FString& DeviceId, TArray<ILauncherProfileLaunchRolePtr>& OutRoles ) = 0;
 
@@ -516,7 +491,6 @@ public:
 	 * Gets the packaging mode.
 	 *
 	 * @return The packaging mode.
-	 *
 	 * @see SetPackagingMode
 	 */
 	virtual ELauncherProfilePackagingModes::Type GetPackagingMode( ) const = 0;
@@ -525,7 +499,6 @@ public:
 	 * Gets the packaging directory.
 	 *
 	 * @return The packaging directory.
-	 *
 	 * @see SetPackageDirectory
 	 */
 	virtual FString GetPackageDirectory( ) const = 0;
@@ -552,16 +525,14 @@ public:
 	 * Gets the full path to the Unreal project to use.
 	 *
 	 * @return The path.
-	 *
 	 * @see SetRocketProjactPath
 	 */
 	virtual FString GetProjectPath( ) const = 0;
 
     /**
-     * Gets the timeout time for the cook on the fly server
-     *
-     * @return timeout time
+     * Gets the timeout time for the cook on the fly server.
      *
+     * @return The timeout time.
      * @see SetTimeout
      */
     virtual uint32 GetTimeout() const = 0;
@@ -570,7 +541,6 @@ public:
 	 * Checks whether the game should be built.
 	 *
 	 * @return true if building the game, false otherwise.
-	 *
 	 * @see SetBuildGame
 	 */
 	virtual bool IsBuilding() const = 0;
@@ -579,7 +549,6 @@ public:
 	 * Checks whether incremental cooking is enabled.
 	 *
 	 * @return true if cooking incrementally, false otherwise.
-	 *
 	 * @see SetIncrementalCooking
 	 */
 	virtual bool IsCookingIncrementally( ) const = 0;
@@ -588,7 +557,6 @@ public:
 	 * Checks whether unversioned cooking is enabled.
 	 *
 	 * @return true if cooking unversioned, false otherwise.
-	 *
 	 * @see SetUnversionedCooking
 	 */
 	virtual bool IsCookingUnversioned( ) const = 0;
@@ -597,7 +565,6 @@ public:
 	 * Checks whether the file server's console window should be hidden.
 	 *
 	 * @return true if the file server should be hidden, false otherwise.
-	 *
 	 * @see SetHideFileServer
 	 */
 	virtual bool IsFileServerHidden( ) const = 0;
@@ -606,7 +573,6 @@ public:
 	 * Checks whether the file server is a streaming file server.
 	 *
 	 * @return true if the file server is streaming, false otherwise.
-	 *
 	 * @see SetStreamingFileServer
 	 */
 	virtual bool IsFileServerStreaming( ) const = 0;
@@ -615,7 +581,6 @@ public:
 	 * Checks whether packaging with UnrealPak is enabled.
 	 *
 	 * @return true if UnrealPak is used, false otherwise.
-	 *
 	 * @see SetPackageWithUnrealPak
 	 */
 	virtual bool IsPackingWithUnrealPak( ) const = 0;
@@ -646,60 +611,48 @@ public:
 	/**
 	 * Adds a culture to cook (only used if cooking by the book).
 	 *
-	 * @param CultureName - The name of the culture to cook.
+	 * @param CultureName The name of the culture to cook.
 	 *
-	 * @see ClearCookedCultures
-	 * @see GetCookedCultures
-	 * @see RemoveCookedCulture
+	 * @see ClearCookedCultures, GetCookedCultures, RemoveCookedCulture
 	 */
 	virtual void AddCookedCulture( const FString& CultureName ) = 0;
 
 	/**
 	 * Adds a map to cook (only used if cooking by the book).
 	 *
-	 * @param MapName - The name of the map to cook.
+	 * @param MapName The name of the map to cook.
 	 *
-	 * @see ClearCookedMaps
-	 * @see GetCookedMaps
-	 * @see RemoveCookedMap
+	 * @see ClearCookedMaps, GetCookedMaps, RemoveCookedMap
 	 */
 	virtual void AddCookedMap( const FString& MapName ) = 0;
 
 	/**
 	 * Adds a platform to cook (only used if cooking by the book).
 	 *
-	 * @param PlatformName - The name of the platform to add.
+	 * @param PlatformName The name of the platform to add.
 	 *
-	 * @see ClearCookedPlatforms
-	 * @see GetCookedPlatforms
-	 * @see RemoveCookedPlatform
+	 * @see ClearCookedPlatforms, GetCookedPlatforms, RemoveCookedPlatform
 	 */
 	virtual void AddCookedPlatform( const FString& PlatformName ) = 0;
 
 	/**
 	 * Removes all cooked cultures.
 	 *
-	 * @see AddCookedCulture
-	 * @see GetCookedCulture
-	 * @see RemoveCookedCulture
+	 * @see AddCookedCulture, GetCookedCulture, RemoveCookedCulture
 	 */
 	virtual void ClearCookedCultures( ) = 0;
 
 	/**
 	 * Removes all cooked maps.
 	 *
-	 * @see AddCookedMap
-	 * @see GetCookedMap
-	 * @see RemoveCookedMap
+	 * @see AddCookedMap, GetCookedMap, RemoveCookedMap
 	 */
 	virtual void ClearCookedMaps( ) = 0;
 
 	/**
 	 * Removes all cooked platforms.
 	 *
-	 * @see AddCookedPlatform
-	 * @see GetCookedPlatforms
-	 * @see RemoveCookedPlatform
+	 * @see AddCookedPlatform, GetCookedPlatforms, RemoveCookedPlatform
 	 */
 	virtual void ClearCookedPlatforms( ) = 0;
 
@@ -707,60 +660,46 @@ public:
 	 * Creates a new launch role and adds it to the profile.
 	 *
 	 * @return The created role.
-	 *
-	 * @see GetLaunchRoles
-	 * @see RemoveLaunchRole
+	 * @see GetLaunchRoles, RemoveLaunchRole
 	 */
 	virtual ILauncherProfileLaunchRolePtr CreateLaunchRole( ) = 0;
 
 	/**
 	 * Removes a cooked culture.
 	 *
-	 * @param CultureName - The name of the culture to remove.
-	 *
-	 * @see AddCookedCulture
-	 * @see ClearCookedCultures
-	 * @see GetCookedCultures
+	 * @param CultureName The name of the culture to remove.
+	 * @see AddCookedCulture, ClearCookedCultures, GetCookedCultures
 	 */
 	virtual void RemoveCookedCulture( const FString& CultureName ) = 0;
 
 	/**
 	 * Removes a cooked map.
 	 *
-	 * @param MapName - The name of the map to remove.
-	 *
-	 * @see AddCookedMap
-	 * @see ClearCookedMaps
-	 * @see GetCookedMaps
+	 * @param MapName The name of the map to remove.
+	 * @see AddCookedMap, ClearCookedMaps, GetCookedMaps
 	 */
 	virtual void RemoveCookedMap( const FString& MapName ) = 0;
 
 	/**
 	 * Removes a platform from the cook list.
 	 *
-	 * @param PlatformName - The name of the platform to remove.
-	 *
-	 * @see AddBuildPlatform
-	 * @see ClearCookedPlatforms
-	 * @see GetBuildPlatforms
+	 * @param PlatformName The name of the platform to remove.
+	 * @see AddBuildPlatform, ClearCookedPlatforms, GetBuildPlatforms
 	 */
 	virtual void RemoveCookedPlatform( const FString& PlatformName ) = 0;
 
 	/**
 	 * Removes the given launch role from the profile.
 	 *
-	 * @param Role - The role to remove.
-	 *
-	 * @see CreateLaunchRole
-	 * @see GetLaunchRoles
+	 * @param Role The role to remove.
+	 * @see CreateLaunchRole, GetLaunchRoles
 	 */
 	virtual void RemoveLaunchRole( const ILauncherProfileLaunchRoleRef& Role ) = 0;
 
 	/**
 	 * Sets whether to build the game.
 	 *
-	 * @param Build - Whether the game should be built.
-	 *
+	 * @param Build Whether the game should be built.
 	 * @see IsBuilding
 	 */
 	virtual void SetBuildGame( bool Build ) = 0;
@@ -768,8 +707,7 @@ public:
 	/**
 	 * Sets the build configuration.
 	 *
-	 * @param ConfigurationName - The build configuration name to set.
-	 *
+	 * @param ConfigurationName The build configuration name to set.
 	 * @see GetBuildConfigurationName
 	 */
 	virtual void SetBuildConfiguration( EBuildConfigurations::Type Configuration ) = 0;
@@ -777,8 +715,7 @@ public:
 	/**
 	 * Sets the build configuration of the cooker.
 	 *
-	 * @param Configuration - The cooker's build configuration to set.
-	 *
+	 * @param Configuration The cooker's build configuration to set.
 	 * @see GetBuildConfigurationName
 	 */
 	virtual void SetCookConfiguration( EBuildConfigurations::Type Configuration ) = 0;
@@ -786,8 +723,7 @@ public:
 	/**
 	 * Sets the cook mode.
 	 *
-	 * @param Mode - The cook mode.
-	 *
+	 * @param Mode The cook mode.
 	 * @see GetCookMode
 	 */
 	virtual void SetCookMode( ELauncherProfileCookModes::Type Mode ) = 0;
@@ -795,8 +731,7 @@ public:
 	/**
 	 * Sets the cook options.
 	 *
-	 * @param Options - The cook options.
-	 *
+	 * @param Options The cook options.
 	 * @see GetCookOptions
 	 */
 	virtual void SetCookOptions(const FString& Options) = 0;
@@ -804,8 +739,7 @@ public:
 	/**
 	 * Sets whether to pack with UnrealPak.
 	 *
-	 * @param UseUnrealPak - Whether UnrealPak should be used.
-	 *
+	 * @param UseUnrealPak Whether UnrealPak should be used.
 	 * @see IsPackingWithUnrealPak
 	 */
 	virtual void SetDeployWithUnrealPak( bool UseUnrealPak ) = 0;
@@ -813,8 +747,7 @@ public:
 	/**
 	 * Sets the device group to deploy to.
 	 *
-	 * @param DeviceGroup - The device group, or NULL to reset this setting.
-	 *
+	 * @param DeviceGroup The device group, or NULL to reset this setting.
 	 * @see GetDeployedDeviceGroup
 	 */
 	virtual void SetDeployedDeviceGroup( const ILauncherDeviceGroupPtr& DeviceGroup ) = 0;
@@ -822,8 +755,7 @@ public:
 	/**
 	 * Sets the deployment mode.
 	 *
-	 * @param Mode - The deployment mode to set.
-	 *
+	 * @param Mode The deployment mode to set.
 	 * @see GetDeploymentMode
 	 */
 	virtual void SetDeploymentMode( ELauncherProfileDeploymentModes::Type Mode ) = 0;
@@ -831,8 +763,7 @@ public:
     /**
      * Sets the cook on the fly close mode
      *
-     * @param Close - the close mode to set
-     *
+     * @param Close the close mode to set.
      * @see GetForceClose
      */
     virtual void SetForceClose( bool Close ) = 0;
@@ -840,8 +771,7 @@ public:
 	/**
 	 * Sets whether to hide the file server's console window.
 	 *
-	 * @param Hide - Whether to hide the window.
-	 *
+	 * @param Hide Whether to hide the window.
 	 * @see GetHideFileServerWindow
 	 */
 	virtual void SetHideFileServerWindow( bool Hide ) = 0;
@@ -849,8 +779,7 @@ public:
 	/**
 	 * Sets incremental cooking.
 	 *
-	 * @param Incremental - Whether cooking should be incremental.
-	 *
+	 * @param Incremental Whether cooking should be incremental.
 	 * @see IsCookingIncrementally
 	 */
 	virtual void SetIncrementalCooking( bool Incremental ) = 0;
@@ -858,8 +787,7 @@ public:
 	/**
 	 * Sets the launch mode.
 	 *
-	 * @param Mode - The launch mode to set.
-	 *
+	 * @param Mode The launch mode to set.
 	 * @see GetLaunchMode
 	 */
 	virtual void SetLaunchMode( ELauncherProfileLaunchModes::Type Mode ) = 0;
@@ -867,8 +795,7 @@ public:
 	/**
 	 * Sets the packaging mode.
 	 *
-	 * @param Mode - The packaging mode to set.
-	 *
+	 * @param Mode The packaging mode to set.
 	 * @see GetPackagingMode
 	 */
 	virtual void SetPackagingMode( ELauncherProfilePackagingModes::Type Mode ) = 0;
@@ -876,8 +803,7 @@ public:
 	/**
 	 * Sets the packaging directory.
 	 *
-	 * @param Dir - The packaging directory to set.
-	 *
+	 * @param Dir The packaging directory to set.
 	 * @see GetPackageDirectory
 	 */
 	virtual void SetPackageDirectory( const FString& Dir ) = 0;
@@ -885,20 +811,17 @@ public:
 	/**
 	 * Sets whether this profile specifies the a project.
 	 *
-	 * @param Specified - Whether a project is specified.
+	 * @param Specified Whether a project is specified.
 	 */
 	virtual void SetProjectSpecified(bool Specified) = 0;
 
-	/**
-	* Notifies the profile that the fallback project path changed.
-	*/
+	/** Notifies the profile that the fallback project path changed. */
 	virtual void FallbackProjectUpdated() = 0;
 
 	/**
 	 * Sets the path to the Rocket project to use.
 	 *
-	 * @param Path - The full path to the project.
-	 *
+	 * @param Path The full path to the project.
 	 * @see GetRocketProjectPath
 	 */
 	virtual void SetProjectPath( const FString& Path ) = 0;
@@ -906,8 +829,7 @@ public:
 	/**
 	 * Sets whether to use a streaming file server.
 	 *
-	 * @param Streaming - Whether a streaming server should be used.
-	 *
+	 * @param Streaming Whether a streaming server should be used.
 	 * @see GetStreamingFileServer
 	 */
 	virtual void SetStreamingFileServer( bool Streaming ) = 0;
@@ -915,8 +837,7 @@ public:
     /**
      * Sets the cook on the fly server timeout
      *
-     * @param InTime - amount of time to wait before timing out
-     *
+     * @param InTime Amount of time to wait before timing out.
      * @see GetTimeout
      */
     virtual void SetTimeout(uint32 InTime) = 0;
@@ -924,22 +845,20 @@ public:
 	/**
 	 * Sets unversioned cooking.
 	 *
-	 * @param Unversioned - Whether cooking is unversioned.
-	 *
+	 * @param Unversioned Whether cooking is unversioned.
 	 * @see IsCookingUnversioned
 	 */
 	virtual void SetUnversionedCooking( bool Unversioned ) = 0;
 
-
 	/**
-	 * Accesses delegate used when the project changes
+	 * Accesses delegate used when the project changes.
+	 *
+	 * @return The delegate.
 	 */
 	virtual FOnProfileProjectChanged& OnProjectChanged() = 0;
 
 public:
 
-	/**
-	 * Virtual destructor.
-	 */
+	/** Virtual destructor. */
 	virtual ~ILauncherProfile( ) { }
 };
diff --git a/Engine/Source/Developer/LauncherServices/Public/Interfaces/ILauncherProfileLaunchRole.h b/Engine/Source/Developer/LauncherServices/Public/Interfaces/ILauncherProfileLaunchRole.h
--- a/Engine/Source/Developer/LauncherServices/Public/Interfaces/ILauncherProfileLaunchRole.h
+++ b/Engine/Source/Developer/LauncherServices/Public/Interfaces/ILauncherProfileLaunchRole.h
@@ -78,7 +78,6 @@ public:
 	 * Gets optional command line parameters to launch with.
 	 *
 	 * @return The command line string.
-	 *
 	 * @see SetCommandLine
 	 */
 	virtual const FString& GetCommandLine( ) const = 0;
@@ -87,7 +86,6 @@ public:
 	 * Gets the initial culture to launch with.
 	 *
 	 * @return Culture name.
-	 *
 	 * @see SetInitialCulture
 	 */
 	virtual const FString& GetInitialCulture( ) const = 0;
@@ -96,7 +94,6 @@ public:
 	 * Gets the initial map to launch with.
 	 *
 	 * @return Map name.
-	 *
 	 * @see SetInitialMap
 	 */
 	virtual const FString& GetInitialMap( ) const = 0;
@@ -119,7 +116,6 @@ public:
 	 * Checks whether vertical sync is enabled.
 	 *
 	 * @return true if VSync is enabled, false otherwise.
-	 *
 	 * @see SetVsyncEnabled
 	 */
 	virtual bool IsVsyncEnabled( ) const = 0;
@@ -127,15 +123,14 @@ public:
 	/**
 	 * Serializes the role from or into the specified archive.
 	 *
-	 * @param Archive - The archive to serialize from or into.
+	 * @param Archive The archive to serialize from or into.
 	 */
 	virtual void Serialize( FArchive& Archive ) = 0;
 
 	/**
 	 * Sets optional command line parameters to launch with.
 	 *
-	 * @param NewCommandLine - Command line string.
-	 *
+	 * @param NewCommandLine Command line string.
 	 * @see GetCommandLine
 	 */
 	virtual void SetCommandLine( const FString& NewCommandLine ) = 0;
@@ -143,8 +138,7 @@ public:
 	/**
 	 * Sets the initial culture to launch with.
 	 *
-	 * @param CultureName - The culture name.
-	 *
+	 * @param CultureName The culture name.
 	 * @see GetInitialCulture
 	 */
 	virtual void SetInitialCulture( const FString& CultureName ) = 0;
@@ -152,8 +146,7 @@ public:
 	/**
 	 * Sets the initial map to launch with.
 	 *
-	 * @param MapName - The name of the map.
-	 *
+	 * @param MapName The name of the map.
 	 * @see GetInitialMap
 	 */
 	virtual void SetInitialMap( const FString& MapName ) = 0;
@@ -161,15 +154,14 @@ public:
 	/**
 	 * Sets the role instance type (i.e. client, server, etc).
 	 *
-	 * @param InInstanceType - The instance type to set.
+	 * @param InInstanceType The instance type to set.
 	 */
 	virtual void SetInstanceType( ELauncherProfileRoleInstanceTypes::Type InInstanceType ) = 0;
 
 	/**
 	 * Sets the name of this role.
 	 *
-	 * @param NewName - The name to set.
-	 *
+	 * @param NewName The name to set.
 	 * @see GetName
 	 */
 	virtual void SetName( const FString& NewName ) = 0;
@@ -177,16 +169,13 @@ public:
 	/**
 	 * Sets whether vertical sync should be enabled.
 	 *
-	 * @param Enabled - Whether VSync is enabled.
-	 *
+	 * @param Enabled Whether VSync is enabled.
 	 * @see IsVsyncEnabled
 	 */
 	virtual void SetVsyncEnabled( bool Enabled ) = 0;
 
 public:
 
-	/**
-	 * Virtual destructor.
-	 */
+	/** Virtual destructor. */
 	virtual ~ILauncherProfileLaunchRole( ) { }
 };
diff --git a/Engine/Source/Developer/LauncherServices/Public/Interfaces/ILauncherProfileManager.h b/Engine/Source/Developer/LauncherServices/Public/Interfaces/ILauncherProfileManager.h
--- a/Engine/Source/Developer/LauncherServices/Public/Interfaces/ILauncherProfileManager.h
+++ b/Engine/Source/Developer/LauncherServices/Public/Interfaces/ILauncherProfileManager.h
@@ -49,7 +49,7 @@ public:
 	/**
 	 * Adds the given device group.
 	 *
-	 * @param DeviceGroup - The group to add.
+	 * @param DeviceGroup The group to add.
 	 */
 	virtual void AddDeviceGroup( const ILauncherDeviceGroupRef& DeviceGroup ) = 0;
 
@@ -77,22 +77,19 @@ public:
 	/**
 	 * Gets the device group with the specified identifier.
 	 *
-	 * @param GroupId - The unique identifier of the group to get.
-	 *
-	 * @return A shared pointer to the group, or NULL if the group was not found.
+	 * @param GroupId The unique identifier of the group to get.
+	 * @return A shared pointer to the group, or nullptr if the group was not found.
 	 */
 	virtual ILauncherDeviceGroupPtr GetDeviceGroup( const FGuid& GroupId ) const = 0;
 
 	/** 
 	 * Deletes the specified device group.
 	 *
-	 *( @param DeviceGroup - The group to remove.
+	 * @param DeviceGroup The group to remove.
 	 */
 	virtual void RemoveDeviceGroup( const ILauncherDeviceGroupRef& DeviceGroup ) = 0;
 
-	/**
-	 * Saves all the device groups to a config file
-	 */
+	/** Saves all the device groups to a config file */
 	virtual void SaveDeviceGroups() = 0;
 
 public:
@@ -100,8 +97,7 @@ public:
 	/**
 	 * Finds or Adds then returns a simple profile for the specified Device
 	 *
-	 * @param DeviceName - The name of the device we want the simple profile for.
-	 *
+	 * @param DeviceName The name of the device we want the simple profile for.
 	 * @return The simple profile for the specified device.
 	 */
 	virtual ILauncherSimpleProfilePtr FindOrAddSimpleProfile(const FString& DeviceName) = 0;
@@ -109,16 +105,15 @@ public:
 	/**
 	 * Gets the simple profile for the specified device.
 	 *
-	 * @param DeviceName - The name of the device we want the simple profile for.
-	 *
-	 * @return The simple profile for the specified device, or NULL if the simple profile doesn't exist.
+	 * @param DeviceName The name of the device we want the simple profile for.
+	 * @return The simple profile for the specified device, or nullptr if the simple profile doesn't exist.
 	 */
 	virtual ILauncherSimpleProfilePtr FindSimpleProfile(const FString& DeviceName) = 0;
 
 	/**
 	 * Deletes the given simple profile.
 	 *
-	 * @param Profile - The simple profile to delete.
+	 * @param Profile The simple profile to delete.
 	 */
 	virtual void RemoveSimpleProfile(const ILauncherSimpleProfileRef& SimpleProfile) = 0;
 
@@ -134,8 +129,7 @@ public:
 	/**
 	 * Creates a new profile but does not add it to the internal tracking.
 	 *
-	 * @param ProfileName - The name of the profile to create.
-	 *
+	 * @param ProfileName The name of the profile to create.
 	 * @return The new profile created.
 	 */
 	virtual ILauncherProfileRef CreateUnsavedProfile(FString ProfileName) = 0;
@@ -146,17 +140,15 @@ public:
 	 * If a profile with the same identifier already exists in the profile
 	 * collection, it will be deleted before the given profile is added.
 	 *
-	 * @param Profile - The profile to add.
+	 * @param Profile The profile to add.
 	 */
 	virtual void AddProfile( const ILauncherProfileRef& Profile ) = 0;
 
 	/**
 	 * Gets the profile with the specified name.
 	 *
-	 * @param ProfileName - The name of the profile to get.
-	 *
-	 * @return The profile, or NULL if the profile doesn't exist.
-	 *
+	 * @param ProfileName The name of the profile to get.
+	 * @return The profile, or nullptr if the profile doesn't exist.
 	 * @see GetProfile
 	 */
 	virtual ILauncherProfilePtr FindProfile( const FString& ProfileName ) = 0;
@@ -173,10 +165,8 @@ public:
 	/**
 	 * Gets the profile with the specified identifier.
 	 *
-	 * @param ProfileId - The identifier of the profile to get.
-	 *
-	 * @return The profile, or NULL if the profile doesn't exist.
-	 *
+	 * @param ProfileId The identifier of the profile to get.
+	 * @return The profile, or nullptr if the profile doesn't exist.
 	 * @see FindProfile
 	 */
 	virtual ILauncherProfilePtr GetProfile( const FGuid& ProfileId ) const = 0;
@@ -187,28 +177,24 @@ public:
 	 * The loaded profile is NOT automatically added to the profile manager.
 	 * Use AddProfile() to add it to the collection.
 	 *
-	 * @param Archive - The archive to load from.
-	 *
-	 * @return The loaded profile, or NULL if loading failed.
-	 *
-	 * @see AddProfile
-	 * @see SaveProfile
+	 * @param Archive The archive to load from.
+	 * @return The loaded profile, or nullptr if loading failed.
+	 * @see AddProfile, SaveProfile
 	 */
 	virtual ILauncherProfilePtr LoadProfile( FArchive& Archive ) = 0;
 
 	/**
 	 * Deletes the given profile.
 	 *
-	 * @param Profile - The profile to delete.
+	 * @param Profile The profile to delete.
 	 */
 	virtual void RemoveProfile( const ILauncherProfileRef& Profile ) = 0;
 
 	/**
 	 * Saves the given profile to the specified archive.
 	 *
-	 * @param Profile - The profile to save.
-	 * @param Archive - The archive to save to.
-	 *
+	 * @param Profile The profile to save.
+	 * @param Archive The archive to save to.
 	 * @see LoadProfile
 	 */
 	virtual void SaveProfile(const ILauncherProfileRef& Profile) = 0;
@@ -249,7 +235,6 @@ public:
 	 * Gets the full path to the Unreal project to use.
 	 *
 	 * @return The path.
-	 *
 	 * @see SetProjectPath
 	 */
 	virtual FString GetProjectPath() const = 0;
@@ -257,8 +242,7 @@ public:
 	/**
 	 * Sets the path to the Unreal project to use.
 	 *
-	 * @param Path - The full path to the project.
-	 *
+	 * @param Path The full path to the project.
 	 * @see GetProjectPath
 	 */
 	virtual void SetProjectPath(const FString& InProjectPath) = 0;
@@ -295,8 +279,6 @@ public:
 
 public:
 
-	/**
-	 * Virtual destructor.
-	 */
+	/** Virtual destructor. */
 	virtual ~ILauncherProfileManager( ) { }
 };
diff --git a/Engine/Source/Developer/LauncherServices/Public/Interfaces/ILauncherServicesModule.h b/Engine/Source/Developer/LauncherServices/Public/Interfaces/ILauncherServicesModule.h
--- a/Engine/Source/Developer/LauncherServices/Public/Interfaces/ILauncherServicesModule.h
+++ b/Engine/Source/Developer/LauncherServices/Public/Interfaces/ILauncherServicesModule.h
@@ -53,11 +53,8 @@ public:
 	virtual FLauncherServicesSDKNotInstalled& OnLauncherServicesSDKNotInstalled( ) = 0;
 	virtual void BroadcastLauncherServicesSDKNotInstalled(const FString& PlatformName, const FString& DocLink) = 0;
 
-
 public:
 
-	/**
-	 * Virtual destructor.
-	 */
+	/** Virtual destructor. */
 	virtual ~ILauncherServicesModule( ) { }
 };
diff --git a/Engine/Source/Developer/LauncherServices/Public/Interfaces/ILauncherTask.h b/Engine/Source/Developer/LauncherServices/Public/Interfaces/ILauncherTask.h
--- a/Engine/Source/Developer/LauncherServices/Public/Interfaces/ILauncherTask.h
+++ b/Engine/Source/Developer/LauncherServices/Public/Interfaces/ILauncherTask.h
@@ -1,9 +1,5 @@
 // Copyright 1998-2014 Epic Games, Inc. All Rights Reserved.
 
-/*=============================================================================
-	ILauncherTask.h: Declares the ILauncherTask interface.
-=============================================================================*/
-
 #pragma once
 
 
diff --git a/Engine/Source/Developer/LauncherServices/Public/Interfaces/ILauncherWorker.h b/Engine/Source/Developer/LauncherServices/Public/Interfaces/ILauncherWorker.h
--- a/Engine/Source/Developer/LauncherServices/Public/Interfaces/ILauncherWorker.h
+++ b/Engine/Source/Developer/LauncherServices/Public/Interfaces/ILauncherWorker.h
@@ -106,8 +106,6 @@ public:
 
 protected:
 
-	/**
-	 * Virtual destructor.
-	 */
+	/** Virtual destructor. */
 	virtual ~ILauncherWorker( ) { }
 };
diff --git a/Engine/Source/Developer/ProfilerMessages/Classes/ProfilerServiceMessages.h b/Engine/Source/Developer/ProfilerMessages/Classes/ProfilerServiceMessages.h
--- a/Engine/Source/Developer/ProfilerMessages/Classes/ProfilerServiceMessages.h
+++ b/Engine/Source/Developer/ProfilerMessages/Classes/ProfilerServiceMessages.h
@@ -1,9 +1,5 @@
 // Copyright 1998-2014 Epic Games, Inc. All Rights Reserved.
 
-/*=============================================================================
-	ProfilerServiceMessages.h: Declares message types used by Profiler services.
-=============================================================================*/
-
 #pragma once
 
 #include "ProfilerServiceMessages.generated.h"
@@ -568,16 +564,3 @@ struct TStructOpsTypeTraits<FProfilerServiceRequest> : public TStructOpsTypeTrai
 		WithMessageHandling = true
 	};
 };
-
-
-/* Dummy class
- *****************************************************************************/
-
-UCLASS()
-class UProfilerServiceMessages
-	: public UObject
-{
-public:
-
-	GENERATED_UCLASS_BODY()
-};
diff --git a/Engine/Source/Developer/ProfilerMessages/Private/ProfilerMessagesModule.cpp b/Engine/Source/Developer/ProfilerMessages/Private/ProfilerMessagesModule.cpp
--- a/Engine/Source/Developer/ProfilerMessages/Private/ProfilerMessagesModule.cpp
+++ b/Engine/Source/Developer/ProfilerMessages/Private/ProfilerMessagesModule.cpp
@@ -1,13 +1,8 @@
 // Copyright 1998-2014 Epic Games, Inc. All Rights Reserved.
 
-/*=============================================================================
-	ProfilerMessagesModule.cpp: Implements the FProfilerMessagesModule class.
-=============================================================================*/
-
 #include "ProfilerMessagesPrivatePCH.h"
 
 
-
 /**
  * Implements the ProfilerMessages module.
  */
@@ -27,10 +22,4 @@ public:
 };
 
 
-// Dummy class initialization
-UProfilerServiceMessages::UProfilerServiceMessages( const class FPostConstructInitializeProperties& PCIP )
-	: Super(PCIP)
-{ }
-
-
 IMPLEMENT_MODULE(FProfilerMessagesModule, ProfilerMessages);
diff --git a/Engine/Source/Developer/ProfilerMessages/Public/ProfilerMessages.h b/Engine/Source/Developer/ProfilerMessages/Public/ProfilerMessages.h
--- a/Engine/Source/Developer/ProfilerMessages/Public/ProfilerMessages.h
+++ b/Engine/Source/Developer/ProfilerMessages/Public/ProfilerMessages.h
@@ -1,9 +1,5 @@
 // Copyright 1998-2014 Epic Games, Inc. All Rights Reserved.
 
-/*=============================================================================
-	ProfilerMessages.h: ProfilerMessages module public header file.
-=============================================================================*/
-
 #pragma once
 
 
diff --git a/Engine/Source/Developer/ProfilerService/Private/ProfilerServiceManager.cpp b/Engine/Source/Developer/ProfilerService/Private/ProfilerServiceManager.cpp
--- a/Engine/Source/Developer/ProfilerService/Private/ProfilerServiceManager.cpp
+++ b/Engine/Source/Developer/ProfilerService/Private/ProfilerServiceManager.cpp
@@ -1,16 +1,14 @@
 // Copyright 1998-2014 Epic Games, Inc. All Rights Reserved.
 
-/*=============================================================================
-	ProfilerServiceModule.cpp: Implements the FProfilerServiceModule class.
-=============================================================================*/
-
 #include "ProfilerServicePrivatePCH.h"
 
 #include "StatsData.h"
 #include "SecureHash.h"
 
+
 DEFINE_LOG_CATEGORY_STATIC(LogProfile, Log, All);
 
+
 /** 
  * Thread used to read, prepare and send files through the message bus.
  * Supports resending bad file chunks and basic synchronization between service and client.
diff --git a/Engine/Source/Developer/ProfilerService/Private/ProfilerServiceManager.h b/Engine/Source/Developer/ProfilerService/Private/ProfilerServiceManager.h
--- a/Engine/Source/Developer/ProfilerService/Private/ProfilerServiceManager.h
+++ b/Engine/Source/Developer/ProfilerService/Private/ProfilerServiceManager.h
@@ -1,9 +1,5 @@
 // Copyright 1998-2014 Epic Games, Inc. All Rights Reserved.
 
-/*=============================================================================
-	ProfilerServiceManager.h: Declares the FProfilerServiceManager class.
-=============================================================================*/
-
 #pragma once
 
 
@@ -28,6 +24,7 @@ struct FClientData
 #endif
 };
 
+
 /**
  * Implements the Profile Service Manager
  */
diff --git a/Engine/Source/Developer/ProfilerService/Private/ProfilerServiceModule.cpp b/Engine/Source/Developer/ProfilerService/Private/ProfilerServiceModule.cpp
--- a/Engine/Source/Developer/ProfilerService/Private/ProfilerServiceModule.cpp
+++ b/Engine/Source/Developer/ProfilerService/Private/ProfilerServiceModule.cpp
@@ -1,13 +1,9 @@
 // Copyright 1998-2014 Epic Games, Inc. All Rights Reserved.
 
-/*=============================================================================
-	ProfilerServiceModule.cpp: Implements the FProfilerServiceModule class.
-=============================================================================*/
-
 #include "ProfilerServicePrivatePCH.h"
-
 #include "ModuleManager.h"
 
+
 /**
  * Implements the ProfilerService module
  */
diff --git a/Engine/Source/Developer/ProfilerService/Private/ProfilerServicePrivatePCH.h b/Engine/Source/Developer/ProfilerService/Private/ProfilerServicePrivatePCH.h
--- a/Engine/Source/Developer/ProfilerService/Private/ProfilerServicePrivatePCH.h
+++ b/Engine/Source/Developer/ProfilerService/Private/ProfilerServicePrivatePCH.h
@@ -1,9 +1,5 @@
 // Copyright 1998-2014 Epic Games, Inc. All Rights Reserved.
 
-/*=============================================================================
-	ProfilerServicePrivatePCH.h: Pre-compiled header file for the ProfileSession module.
-=============================================================================*/
-
 #pragma once
 
 
diff --git a/Engine/Source/Developer/ProfilerService/Public/Interfaces/IProfilerServiceManager.h b/Engine/Source/Developer/ProfilerService/Public/Interfaces/IProfilerServiceManager.h
--- a/Engine/Source/Developer/ProfilerService/Public/Interfaces/IProfilerServiceManager.h
+++ b/Engine/Source/Developer/ProfilerService/Public/Interfaces/IProfilerServiceManager.h
@@ -1,21 +1,15 @@
 // Copyright 1998-2014 Epic Games, Inc. All Rights Reserved.
 
-/*=============================================================================
-	IProfilerServiceManager.h: Declares the IProfileServiceManager interface.
-=============================================================================*/
-
 #pragma once
 
-/**
- * Type definition for shared pointers to instances of IProfilerServiceManager.
- */
+
+/** Type definition for shared pointers to instances of IProfilerServiceManager. */
 typedef TSharedPtr<class IProfilerServiceManager> IProfilerServiceManagerPtr;
 
-/**
- * Type definition for shared references to instances of IProfilerService.
- */
+/** Type definition for shared references to instances of IProfilerService. */
 typedef TSharedRef<class IProfilerServiceManager> IProfilerServiceManagerRef;
 
+
 /**
  * Type definition for the profiler data packet.
  */
@@ -33,18 +27,21 @@ public:
 	uint32 CallsPerFrame;
 };
 
+
 struct FProfilerFloatAccumulator
 {
 	int32 StatId;
 	float Value;
 };
 
+
 struct FProfilerCountAccumulator
 {
 	int32 StatId;
 	uint32 Value;
 };
 
+
 struct FProfilerCycleGraph
 {
 public:
@@ -56,6 +53,7 @@ public:
 	TArray<FProfilerCycleGraph> Children;
 };
 
+
 FORCEINLINE FArchive& operator<<(FArchive& Ar, FProfilerCycleCounter& Data)
 {
 	Ar << Data.StatId;
@@ -70,6 +68,7 @@ FORCEINLINE FArchive& operator<<(FArchive& Ar, FProfilerCycleCounter& Data)
 
 	return Ar;
 }
+
 	
 FORCEINLINE FArchive& operator<<(FArchive& Ar, FProfilerFloatAccumulator& Data)
 {
@@ -79,6 +78,7 @@ FORCEINLINE FArchive& operator<<(FArchive& Ar, FProfilerFloatAccumulator& Data)
 	return Ar;
 }
 
+
 FORCEINLINE FArchive& operator<<(FArchive& Ar, FProfilerCountAccumulator& Data)
 {
 	Ar << Data.StatId;
@@ -87,6 +87,7 @@ FORCEINLINE FArchive& operator<<(FArchive& Ar, FProfilerCountAccumulator& Data)
 	return Ar;
 }
 
+
 FORCEINLINE FArchive& operator<<(FArchive& Ar, FProfilerCycleGraph& Data)
 {
 	Ar << Data.StatId;
@@ -100,6 +101,7 @@ FORCEINLINE FArchive& operator<<(FArchive& Ar, FProfilerCycleGraph& Data)
 	return Ar;
 }
 
+
 /**
  * Type definition for a frame profiler data
  */
@@ -126,6 +128,7 @@ FORCEINLINE FArchive& operator<<(FArchive& Ar, FProfilerDataFrame& Data)
 	return Ar;
 }
 
+
 struct FStatDescription
 {
 	/** Default constructor. */
@@ -220,6 +223,7 @@ FORCEINLINE FArchive& operator<<(FArchive& Ar, FStatMetaData& Data)
 	return Ar;
 }
 
+
 /**
  * Profiler file type
  */
@@ -354,6 +358,7 @@ FORCEINLINE FArchive& operator<<(FArchive& Ar, FStatMessage& Data)
 }
 #endif
 
+
 /**
  * Profiler service request type
  */
@@ -375,6 +380,7 @@ DECLARE_MULTICAST_DELEGATE_TwoParams(FProfilerDataDelegate, const FGuid&, const
 DECLARE_MULTICAST_DELEGATE_TwoParams(FProfilerData2Delegate, const FGuid&, TArray<FStatMessage>& );
 #endif
 
+
 /** Enumerates profiler file chunk types. */
 namespace EProfilerFileChunkType
 {
@@ -391,6 +397,7 @@ namespace EProfilerFileChunkType
 	};
 }
 
+
 /** Class that describes basic information about one file chunk. */
 struct FProfilerFileChunkHeader
 {
@@ -467,63 +474,55 @@ class IProfilerServiceManager
 public:
 
 	/**
-	 * Sends profiler cycle counter to registered clients
+	 * Sends profiler cycle counter to registered clients.
 	 *
-	 * @param Data - The profiler data to be sent
+	 * @param Data The profiler data to be sent.
 	 */
 	virtual void SendData(FProfilerCycleCounter& Data) = 0;
 
 	/**
-	 * Sends profiler float accumulator to registered clients
+	 * Sends profiler float accumulator to registered clients.
 	 *
-	 * @param Data - The profiler data to be sent
+	 * @param Data The profiler data to be sent.
 	 */
 	virtual void SendData(FProfilerFloatAccumulator& Data) = 0;
 
 	/**
-	 * Sends profiler count accumulator to registered clients
+	 * Sends profiler count accumulator to registered clients.
 	 *
-	 * @param Data - The profiler data to be sent
+	 * @param Data The profiler data to be sent.
 	 */
 	virtual void SendData(FProfilerCountAccumulator& Data) = 0;
 
 	/**
-	 * Sends profiler cycle graph to registered clients
+	 * Sends profiler cycle graph to registered clients.
 	 *
-	 * @param Data - The profiler data to be sent
+	 * @param Data The profiler data to be sent.
 	 */
 	virtual void SendData(FProfilerCycleGraph& Data) = 0;
 
 	/**
-	 * Determines if we are capturing data for clients
+	 * Determines if we are capturing data for clients.
 	 */
 	virtual bool IsCapturing() const = 0;
 
-	/**
-	 * Starts a file capture
-	 */
+	/** Starts a file capture. */
 	virtual void StartCapture() = 0;
 
-	/**
-	 * Stops a file capture
-	 */
+	/** Stops a file capture. */
 	virtual void StopCapture() = 0;
 
-	/**
-	 * Updates the meta data
-	 */
+	/** Updates the meta data. */
 	virtual void UpdateMetaData() = 0;
 
-	/**
-	 * Starts a new frame of data and sends the previous if it exists
-	 */
+	/** Starts a new frame of data and sends the previous if it exists. */
 	virtual void StartFrame(uint32 FrameNumber, double FrameStart) = 0;
 
 	/**
-	 * Gets the description for the given stat id
+	 * Gets the description for the given stat id.
 	 *
-	 * @param StatId - id of the statistic description to retrieve
-	 * @return the FStatMetaData struct with the description
+	 * @param StatId IDd of the statistic description to retrieve.
+	 * @return the FStatMetaData struct with the description.
 	 */
 	virtual FStatMetaData& GetStatMetaData() = 0;
 
@@ -535,8 +534,7 @@ public:
 	virtual FProfilerDataDelegate& OnProfilerData() = 0;
 
 public:
-	/**
-	 * Virtual destructor
-	 */
+
+	/** Virtual destructor. */
 	virtual ~IProfilerServiceManager() { }
 };
diff --git a/Engine/Source/Developer/ProfilerService/Public/Interfaces/IProfilerServiceModule.h b/Engine/Source/Developer/ProfilerService/Public/Interfaces/IProfilerServiceModule.h
--- a/Engine/Source/Developer/ProfilerService/Public/Interfaces/IProfilerServiceModule.h
+++ b/Engine/Source/Developer/ProfilerService/Public/Interfaces/IProfilerServiceModule.h
@@ -1,9 +1,5 @@
 // Copyright 1998-2014 Epic Games, Inc. All Rights Reserved.
 
-/*=============================================================================
-	IProfilerServiceModule.h: Declares the IProfileServiceModule interface.
-=============================================================================*/
-
 #pragma once
 
 
@@ -14,19 +10,17 @@ class IProfilerServiceModule
 	: public IModuleInterface
 {
 public:
+
 	/** 
 	 * Creates a profiler service.
 	 *
-	 * @param Capture - specifies whether to start capturing immediately
-	 *
+	 * @param Capture Specifies whether to start capturing immediately.
 	 * @return A new profiler service.
 	 */
 	virtual IProfilerServiceManagerPtr CreateProfilerServiceManager( ) = 0;
 
 protected:
 
-	/**
-	 * Virtual destructor
-	 */
+	/** Virtual destructor */
 	virtual ~IProfilerServiceModule( ) { }
 };
diff --git a/Engine/Source/Developer/ProfilerService/Public/ProfilerService.h b/Engine/Source/Developer/ProfilerService/Public/ProfilerService.h
--- a/Engine/Source/Developer/ProfilerService/Public/ProfilerService.h
+++ b/Engine/Source/Developer/ProfilerService/Public/ProfilerService.h
@@ -1,9 +1,5 @@
 // Copyright 1998-2014 Epic Games, Inc. All Rights Reserved.
 
-/*=============================================================================
-	ProfilerService.h: Profile Service module public header file.
-=============================================================================*/
-
 #pragma once
 
 
diff --git a/Engine/Source/Developer/ProjectLauncher/Private/Widgets/Build/SProjectLauncherBuildPage.cpp b/Engine/Source/Developer/ProjectLauncher/Private/Widgets/Build/SProjectLauncherBuildPage.cpp
--- a/Engine/Source/Developer/ProjectLauncher/Private/Widgets/Build/SProjectLauncherBuildPage.cpp
+++ b/Engine/Source/Developer/ProjectLauncher/Private/Widgets/Build/SProjectLauncherBuildPage.cpp
@@ -1,9 +1,5 @@
 // Copyright 1998-2014 Epic Games, Inc. All Rights Reserved.
 
-/*=============================================================================
-	SProjectLauncherCookPage.cpp: Implements the SProjectLauncherCookPage class.
-=============================================================================*/
-
 #include "ProjectLauncherPrivatePCH.h"
 
 
@@ -154,8 +150,6 @@ bool SProjectLauncherBuildPage::GenerateDSYMForProject( const FString& ProjectNa
 /* SProjectLauncherBuildPage callbacks
  *****************************************************************************/
 
-
-// Callback for changing the checked state of a platform menu check box. 
 void SProjectLauncherBuildPage::HandleBuildCheckedStateChanged( ESlateCheckBoxState::Type CheckState )
 {
 	ILauncherProfilePtr SelectedProfile = Model->GetSelectedProfile();
@@ -166,7 +160,7 @@ void SProjectLauncherBuildPage::HandleBuildCheckedStateChanged( ESlateCheckBoxSt
 	}
 }
 
-// Callback for determining whether a platform menu entry is checked.
+
 ESlateCheckBoxState::Type SProjectLauncherBuildPage::HandleBuildIsChecked() const
 {
 	ILauncherProfilePtr SelectedProfile = Model->GetSelectedProfile();
@@ -182,11 +176,13 @@ ESlateCheckBoxState::Type SProjectLauncherBuildPage::HandleBuildIsChecked() cons
 	return ESlateCheckBoxState::Unchecked;
 }
 
+
 void SProjectLauncherBuildPage::HandleProfileManagerProfileSelected( const ILauncherProfilePtr& SelectedProfile, const ILauncherProfilePtr& PreviousProfile )
 {
 	// reload settings
 }
 
+
 EVisibility SProjectLauncherBuildPage::HandleBuildPlatformVisibility( ) const
 {
 	ILauncherProfilePtr SelectedProfile = Model->GetSelectedProfile();
@@ -202,7 +198,7 @@ EVisibility SProjectLauncherBuildPage::HandleBuildPlatformVisibility( ) const
 	return EVisibility::Collapsed;
 }
 
-// Callback for pressing the Advanced Setting - Generate DSYM button.
+
 FReply SProjectLauncherBuildPage::HandleGenDSYMClicked()
 {
     ILauncherProfilePtr SelectedProfile = Model->GetSelectedProfile();
@@ -221,6 +217,7 @@ FReply SProjectLauncherBuildPage::HandleGenDSYMClicked()
     return FReply::Handled();
 }
 
+
 bool SProjectLauncherBuildPage::HandleGenDSYMButtonEnabled() const
 {
     ILauncherProfilePtr SelectedProfile = Model->GetSelectedProfile();
@@ -236,6 +233,7 @@ bool SProjectLauncherBuildPage::HandleGenDSYMButtonEnabled() const
     return false;
 }
 
+
 EVisibility SProjectLauncherBuildPage::ShowBuildConfiguration() const
 {
 	ILauncherProfilePtr SelectedProfile = Model->GetSelectedProfile();
@@ -250,6 +248,7 @@ EVisibility SProjectLauncherBuildPage::ShowBuildConfiguration() const
 	}
 }
 
+
 void SProjectLauncherBuildPage::HandleBuildConfigurationSelectorConfigurationSelected(EBuildConfigurations::Type Configuration)
 {
 	ILauncherProfilePtr SelectedProfile = Model->GetSelectedProfile();
@@ -260,6 +259,7 @@ void SProjectLauncherBuildPage::HandleBuildConfigurationSelectorConfigurationSel
 	}
 }
 
+
 FString SProjectLauncherBuildPage::HandleBuildConfigurationSelectorText() const
 {
 	ILauncherProfilePtr SelectedProfile = Model->GetSelectedProfile();
@@ -272,6 +272,7 @@ FString SProjectLauncherBuildPage::HandleBuildConfigurationSelectorText() const
 	return FString();
 }
 
+
 EVisibility SProjectLauncherBuildPage::HandleValidationErrorIconVisibility(ELauncherProfileValidationErrors::Type Error) const
 {
 	ILauncherProfilePtr SelectedProfile = Model->GetSelectedProfile();
@@ -287,4 +288,5 @@ EVisibility SProjectLauncherBuildPage::HandleValidationErrorIconVisibility(ELaun
 	return EVisibility::Hidden;
 }
 
+
 #undef LOCTEXT_NAMESPACE
diff --git a/Engine/Source/Developer/ProjectLauncher/Private/Widgets/Build/SProjectLauncherBuildPage.h b/Engine/Source/Developer/ProjectLauncher/Private/Widgets/Build/SProjectLauncherBuildPage.h
--- a/Engine/Source/Developer/ProjectLauncher/Private/Widgets/Build/SProjectLauncherBuildPage.h
+++ b/Engine/Source/Developer/ProjectLauncher/Private/Widgets/Build/SProjectLauncherBuildPage.h
@@ -16,9 +16,7 @@ public:
 
 public:
 
-	/**
-	 * Destructor.
-	 */
+	/** Destructor. */
 	~SProjectLauncherBuildPage( );
 
 public:
@@ -44,38 +42,38 @@ protected:
 
 private:
 
-	// Callback for changing the checked state of a platform menu check box. 
+	/** Callback for changing the checked state of a platform menu check box. */
 	void HandleBuildCheckedStateChanged( ESlateCheckBoxState::Type CheckState );
 
-	// Callback for determining whether a platform menu entry is checked.
+	/** Callback for determining whether a platform menu entry is checked. */
 	ESlateCheckBoxState::Type HandleBuildIsChecked() const;
 
-	// Callback for changing the selected profile in the profile manager.
+	/** Callback for changing the selected profile in the profile manager. */
 	void HandleProfileManagerProfileSelected( const ILauncherProfilePtr& SelectedProfile, const ILauncherProfilePtr& PreviousProfile );
 
-	// Callback for determining if the build platform list should be displayed
+	/** Callback for determining if the build platform list should be displayed. */
 	EVisibility HandleBuildPlatformVisibility( ) const;
 
-    // Callback for pressing the Advanced Setting - Generate DSYM button.
+    /** Callback for pressing the Advanced Setting - Generate DSYM button. */
     FReply HandleGenDSYMClicked();
 
-	// Callback for getting the enabled state of the Generate DSYM button.
+	/** Callback for getting the enabled state of the Generate DSYM button. */
     bool HandleGenDSYMButtonEnabled() const;
 
-	// Callback for determining if the build configuration should be shown
+	/** Callback for determining if the build configuration should be shown */
 	EVisibility ShowBuildConfiguration() const;
 
-	// Callback for selecting a build configuration.
+	/** Callback for selecting a build configuration. */
 	void HandleBuildConfigurationSelectorConfigurationSelected(EBuildConfigurations::Type Configuration);
 
-	// Callback for getting the content text of the build configuration selector.
+	/** Callback for getting the content text of the build configuration selector. */
 	FString HandleBuildConfigurationSelectorText() const;
 
-	// Callback for determining the visibility of a validation error icon.
+	/** Callback for determining the visibility of a validation error icon. */
 	EVisibility HandleValidationErrorIconVisibility(ELauncherProfileValidationErrors::Type Error) const;
 
 private:
 
-	// Holds a pointer to the data model.
+	/** Holds a pointer to the data model. */
 	FProjectLauncherModelPtr Model;
 };
