Fix an issue where map list was not part of the UAT command line.

diff --git a/Engine/Source/Developer/LauncherServices/Private/Launcher/LauncherWorker.cpp b/Engine/Source/Developer/LauncherServices/Private/Launcher/LauncherWorker.cpp
--- a/Engine/Source/Developer/LauncherServices/Private/Launcher/LauncherWorker.cpp
+++ b/Engine/Source/Developer/LauncherServices/Private/Launcher/LauncherWorker.cpp
@@ -309,6 +309,8 @@ FString FLauncherWorker::CreateUATCommand( const ILauncherProfileRef& InProfile,
 		}
 	}
 
+	UATCommand  += MapList; 
+
 	// build
 	if (InProfile->IsBuilding())
 	{
