Fix for launched game sessions not showing up in the session manager.

diff --git a/Engine/Source/Developer/LauncherServices/Private/Launcher/LauncherWorker.cpp b/Engine/Source/Developer/LauncherServices/Private/Launcher/LauncherWorker.cpp
--- a/Engine/Source/Developer/LauncherServices/Private/Launcher/LauncherWorker.cpp
+++ b/Engine/Source/Developer/LauncherServices/Private/Launcher/LauncherWorker.cpp
@@ -199,7 +199,7 @@ FString FLauncherWorker::CreateUATCommand( const ILauncherProfileRef& InProfile,
 		*InitialMap,
 //		*InstanceId.ToString(),
 		*SessionId.ToString(),
-		FPlatformProcess::UserName(false),
+		FPlatformProcess::UserName(true),
 		*InProfile->GetName());
 
 	// staging directory
