Changed the default configuration for new profiles to be development.

diff --git a/Engine/Source/Developer/LauncherServices/Private/Profiles/LauncherProfile.h b/Engine/Source/Developer/LauncherServices/Private/Profiles/LauncherProfile.h
--- a/Engine/Source/Developer/LauncherServices/Private/Profiles/LauncherProfile.h
+++ b/Engine/Source/Developer/LauncherServices/Private/Profiles/LauncherProfile.h
@@ -567,7 +567,10 @@ public:
 		// Use the locally specified project path is resolving through the root isn't working
 		ProjectSpecified = GetProjectPath().IsEmpty();
 		
-		BuildConfiguration = FApp::GetBuildConfiguration();
+		// I don't use FApp::GetBuildConfiguration() because i don't want the act of running in debug the first time to cause the simple
+		// profiles created for your persistent devices to be in debug. The use might not see this if they don't expand the Advanced options.
+		BuildConfiguration = EBuildConfigurations::Development;
+
 		FInternationalization& I18N = FInternationalization::Get();
 
 		// default build settings
