Documentation corrections.

diff --git a/Engine/Source/Developer/LauncherServices/Private/Profiles/LauncherProfile.h b/Engine/Source/Developer/LauncherServices/Private/Profiles/LauncherProfile.h
--- a/Engine/Source/Developer/LauncherServices/Private/Profiles/LauncherProfile.h
+++ b/Engine/Source/Developer/LauncherServices/Private/Profiles/LauncherProfile.h
@@ -87,7 +87,7 @@ public:
 		Variant = NAME_None;
 		
 		// I don't use FApp::GetBuildConfiguration() because i don't want the act of running in debug the first time to cause the simple
-		// profiles created for your persistent devices to be in debug. The use might not see this if they don't expand the Advanced options.
+		// profiles created for your persistent devices to be in debug. The user might not see this if they don't expand the Advanced options.
 		BuildConfiguration = EBuildConfigurations::Development;
 		
 		CookMode = ELauncherProfileCookModes::OnTheFly;
@@ -567,8 +567,8 @@ public:
 		// Use the locally specified project path is resolving through the root isn't working
 		ProjectSpecified = GetProjectPath().IsEmpty();
 		
-		// I don't use FApp::GetBuildConfiguration() because i don't want the act of running in debug the first time to cause the simple
-		// profiles created for your persistent devices to be in debug. The use might not see this if they don't expand the Advanced options.
+		// I don't use FApp::GetBuildConfiguration() because i don't want the act of running in debug the first time to cause 
+		// profiles the user creates to be in debug. This will keep consistency.
 		BuildConfiguration = EBuildConfigurations::Development;
 
 		FInternationalization& I18N = FInternationalization::Get();
