Front End: Copy repository functionality added back in.

diff --git a/Engine/Source/Developer/LauncherServices/Private/Launcher/LauncherWorker.cpp b/Engine/Source/Developer/LauncherServices/Private/Launcher/LauncherWorker.cpp
--- a/Engine/Source/Developer/LauncherServices/Private/Launcher/LauncherWorker.cpp
+++ b/Engine/Source/Developer/LauncherServices/Private/Launcher/LauncherWorker.cpp
@@ -418,7 +418,26 @@ FString FLauncherWorker::CreateUATCommand( const ILauncherProfileRef& InProfile,
 		switch (InProfile->GetDeploymentMode())
 		{
 		case ELauncherProfileDeploymentModes::CopyRepository:
+			{
+				UATCommand += TEXT(" -skipstage -deploy");
+				UATCommand += CommandLine;
+				UATCommand += StageDirectory;
+				UATCommand += DeviceCommand;
+				UATCommand += AdditionalCommandLine;
+
+				FCommandDesc Desc;
+				FText Command = FText::Format(LOCTEXT("LauncherDeployDesc", "Deploying content for {0}"), FText::FromString(Platforms.RightChop(1)));
+				Desc.Name = "Deploy Task";
+				Desc.Desc = Command.ToString();
+				Desc.EndText = TEXT("********** DEPLOY COMMAND COMPLETED **********");
+				OutCommands.Add(Desc);
+				if (CommandStart.Len() == 0)
+				{
+					CommandStart = TEXT("********** DEPLOY COMMAND STARTED **********");
+				}
+			}
 			break;
+
 		case ELauncherProfileDeploymentModes::FileServer:
 		case ELauncherProfileDeploymentModes::CopyToDevice:
 			{
diff --git a/Engine/Source/Programs/AutomationTool/IOS/IOSPlatform.Automation.cs b/Engine/Source/Programs/AutomationTool/IOS/IOSPlatform.Automation.cs
--- a/Engine/Source/Programs/AutomationTool/IOS/IOSPlatform.Automation.cs
+++ b/Engine/Source/Programs/AutomationTool/IOS/IOSPlatform.Automation.cs
@@ -833,10 +833,6 @@ public class IOSPlatform : Platform
 		return false;
 	}
 
-    public override PakType RequiresPak(ProjectParams Params)
-    {
-        return PakType.Always;
-    }
 
 	public override bool DeployLowerCaseFilenames(bool bUFSFile)
 	{
diff --git a/Engine/Source/Programs/AutomationTool/Scripts/PackageCommand.Automation.cs b/Engine/Source/Programs/AutomationTool/Scripts/PackageCommand.Automation.cs
--- a/Engine/Source/Programs/AutomationTool/Scripts/PackageCommand.Automation.cs
+++ b/Engine/Source/Programs/AutomationTool/Scripts/PackageCommand.Automation.cs
@@ -23,7 +23,7 @@ public partial class Project : CommandUtils
 		{
 			DeployContextList.AddRange(CreateDeploymentContext(Params, true, false));
 		}
-		if (DeployContextList.Count > 0)
+		if (DeployContextList.Count > 0 && !Params.SkipStage)
 		{
 			Log("********** PACKAGE COMMAND STARTED **********");
 
