position of roles command line wasn't inisde the additional commands
TTP347613
#ue4

diff --git a/Engine/Source/Developer/LauncherServices/Private/Launcher/LauncherWorker.cpp b/Engine/Source/Developer/LauncherServices/Private/Launcher/LauncherWorker.cpp
--- a/Engine/Source/Developer/LauncherServices/Private/Launcher/LauncherWorker.cpp
+++ b/Engine/Source/Developer/LauncherServices/Private/Launcher/LauncherWorker.cpp
@@ -195,13 +195,6 @@ FString FLauncherWorker::CreateUATCommand( const ILauncherProfileRef& InProfile,
 	FString CommandLine = FString::Printf(TEXT(" -cmdline=\"%s -Messaging\""),
 		*InitialMap);
 
-	FString AdditionalCommandLine = FString::Printf(TEXT(" -addcmdline=\"%s -SessionId=%s -SessionOwner=%s -SessionName='%s'\""),
-		*InitialMap,
-//		*InstanceId.ToString(),
-		*SessionId.ToString(),
-		FPlatformProcess::UserName(true),
-		*InProfile->GetName());
-
 	// staging directory
 	FString StageDirectory = TEXT("");
 	auto PackageDirectory = InProfile->GetPackageDirectory();
@@ -311,6 +304,15 @@ FString FLauncherWorker::CreateUATCommand( const ILauncherProfileRef& InProfile,
 		DeviceCommand += TEXT(" -device=") + DeviceNames.RightChop(1);
 	}
 
+	// additional commands to be sent to the commandline
+	FString AdditionalCommandLine = FString::Printf(TEXT(" -addcmdline=\"%s -SessionId=%s -SessionOwner=%s -SessionName='%s'%s\""),
+		*InitialMap,
+		//		*InstanceId.ToString(),
+		*SessionId.ToString(),
+		FPlatformProcess::UserName(true),
+		*InProfile->GetName(),
+		*RoleCommands);
+
 	// map list
 	FString MapList = TEXT("");
 	const TArray<FString>& CookedMaps = InProfile->GetCookedMaps();
@@ -465,7 +467,6 @@ FString FLauncherWorker::CreateUATCommand( const ILauncherProfileRef& InProfile,
 		if (InProfile->GetLaunchMode() != ELauncherProfileLaunchModes::DoNotLaunch)
 		{
 			UATCommand += TEXT(" -run");
-			UATCommand += RoleCommands;
 
 			FCommandDesc Desc;
 			FText Command = FText::Format(LOCTEXT("LauncherRunDesc", "Launching on {0}"), FText::FromString(DeviceNames.RightChop(1)));
