missing initial map for launch on
#ue4

diff --git a/Engine/Source/Developer/LauncherServices/Private/Launcher/LauncherWorker.cpp b/Engine/Source/Developer/LauncherServices/Private/Launcher/LauncherWorker.cpp
--- a/Engine/Source/Developer/LauncherServices/Private/Launcher/LauncherWorker.cpp
+++ b/Engine/Source/Developer/LauncherServices/Private/Launcher/LauncherWorker.cpp
@@ -191,10 +191,6 @@ FString FLauncherWorker::CreateUATCommand( const ILauncherProfileRef& InProfile,
 		InitialMap = InProfile->GetCookedMaps()[0];
 	}
 
-	// game command line
-	FString CommandLine = FString::Printf(TEXT(" -cmdline=\"%s -Messaging\""),
-		*InitialMap);
-
 	// staging directory
 	FString StageDirectory = TEXT("");
 	auto PackageDirectory = InProfile->GetPackageDirectory();
@@ -304,10 +300,12 @@ FString FLauncherWorker::CreateUATCommand( const ILauncherProfileRef& InProfile,
 		DeviceCommand += TEXT(" -device=") + DeviceNames.RightChop(1);
 	}
 
+	// game command line
+	FString CommandLine = FString::Printf(TEXT(" -cmdline=\"%s -Messaging\""),
+		*InitialMap);
+
 	// additional commands to be sent to the commandline
-	FString AdditionalCommandLine = FString::Printf(TEXT(" -addcmdline=\"%s -SessionId=%s -SessionOwner=%s -SessionName='%s'%s\""),
-		*InitialMap,
-		//		*InstanceId.ToString(),
+	FString AdditionalCommandLine = FString::Printf(TEXT(" -addcmdline=\"-SessionId=%s -SessionOwner=%s -SessionName='%s'%s\""),
 		*SessionId.ToString(),
 		FPlatformProcess::UserName(true),
 		*InProfile->GetName(),
@@ -316,7 +314,7 @@ FString FLauncherWorker::CreateUATCommand( const ILauncherProfileRef& InProfile,
 	// map list
 	FString MapList = TEXT("");
 	const TArray<FString>& CookedMaps = InProfile->GetCookedMaps();
-	if (CookedMaps.Num() > 0)
+	if (CookedMaps.Num() > 0 && (InProfile->GetCookMode() == ELauncherProfileCookModes::ByTheBook || InProfile->GetCookMode() == ELauncherProfileCookModes::ByTheBookInEditor))
 	{
 		MapList += TEXT(" -map=");
 		for (int32 MapIndex = 0; MapIndex < CookedMaps.Num(); ++MapIndex)
@@ -328,6 +326,10 @@ FString FLauncherWorker::CreateUATCommand( const ILauncherProfileRef& InProfile,
 			}
 		}
 	}
+	else
+	{
+		MapList = TEXT(" -map=") + InitialMap;
+	}
 
 	// build
 	if (InProfile->IsBuilding())
@@ -395,6 +397,7 @@ FString FLauncherWorker::CreateUATCommand( const ILauncherProfileRef& InProfile,
 			{
 				UATCommand += " -nokill";
 			}
+			UATCommand += MapList;
 
 			FCommandDesc Desc;
 			FText Command = LOCTEXT("LauncherCookDesc", "Starting cook on the fly server");
