fix for incorrect map when using Launch On and Cook in Editor
UE-1317
#ue4

diff --git a/Engine/Source/Developer/LauncherServices/Private/Launcher/LauncherWorker.cpp b/Engine/Source/Developer/LauncherServices/Private/Launcher/LauncherWorker.cpp
--- a/Engine/Source/Developer/LauncherServices/Private/Launcher/LauncherWorker.cpp
+++ b/Engine/Source/Developer/LauncherServices/Private/Launcher/LauncherWorker.cpp
@@ -412,6 +412,7 @@ FString FLauncherWorker::CreateUATCommand( const ILauncherProfileRef& InProfile,
 		}
 		break;
 	case ELauncherProfileCookModes::ByTheBookInEditor:
+		UATCommand += MapList;
 	case ELauncherProfileCookModes::DoNotCook:
 		UATCommand += TEXT(" -skipcook");
 		break;
