Change launch platform back to Win64. Binary Rocket does not have Win32 Development executables.

diff --git a/Engine/Source/Developer/LauncherServices/Private/Launcher/LauncherWorker.cpp b/Engine/Source/Developer/LauncherServices/Private/Launcher/LauncherWorker.cpp
--- a/Engine/Source/Developer/LauncherServices/Private/Launcher/LauncherWorker.cpp
+++ b/Engine/Source/Developer/LauncherServices/Private/Launcher/LauncherWorker.cpp
@@ -225,7 +225,7 @@ FString FLauncherWorker::CreateUATCommand( const ILauncherProfileRef& InProfile,
 		}
 		else if (PlatformInfo->TargetPlatformName == FName("WindowsNoEditor") || PlatformInfo->TargetPlatformName == FName("Windows"))
 		{
-			Platforms += TEXT("+Win32");
+			Platforms += TEXT("+Win64");
 		}
 		else if (PlatformInfo->TargetPlatformName == FName("MacNoEditor"))
 		{
