UE-5666 Quick launching crashes for code projects
Properly set build as default if the project is code and user has compiler when using Quick Launch in UFE (similar to editor's Launch On)
#ufe

diff --git a/Engine/Source/Developer/LauncherServices/LauncherServices.Build.cs b/Engine/Source/Developer/LauncherServices/LauncherServices.Build.cs
--- a/Engine/Source/Developer/LauncherServices/LauncherServices.Build.cs
+++ b/Engine/Source/Developer/LauncherServices/LauncherServices.Build.cs
@@ -23,6 +23,7 @@ namespace UnrealBuildTool.Rules
 					"CoreUObject",
 					"DesktopPlatform",
 					"SessionMessages",
+					"SourceCodeAccess",
 					"TargetPlatform",
 					"UnrealEdMessages",
 				}
diff --git a/Engine/Source/Developer/LauncherServices/Private/Launcher/LauncherWorker.cpp b/Engine/Source/Developer/LauncherServices/Private/Launcher/LauncherWorker.cpp
--- a/Engine/Source/Developer/LauncherServices/Private/Launcher/LauncherWorker.cpp
+++ b/Engine/Source/Developer/LauncherServices/Private/Launcher/LauncherWorker.cpp
@@ -2,6 +2,7 @@
 
 #include "LauncherServicesPrivatePCH.h"
 #include "PlatformInfo.h"
+#include "ISourceCodeAccessModule.h"
 
 #define LOCTEXT_NAMESPACE "LauncherWorker"
 
@@ -529,6 +530,20 @@ FString FLauncherWorker::CreateUATCommand( const ILauncherProfileRef& InProfile,
 
 void FLauncherWorker::CreateAndExecuteTasks( const ILauncherProfileRef& InProfile )
 {
+	// check to see if we need to build by default
+	if (!InProfile->HasProjectSpecified())
+	{
+		// TODO: make SouceCodeNavigation and GameProjectGeneration modules available to UFE
+		FString ProjectPath = FPaths::GetPath(InProfile->GetProjectPath());
+		TArray<FString> OutProjectCodeFilenames;
+		IFileManager::Get().FindFilesRecursive(OutProjectCodeFilenames, *(ProjectPath / TEXT("Source")), TEXT("*.h"), true, false, false);
+		IFileManager::Get().FindFilesRecursive(OutProjectCodeFilenames, *(ProjectPath / TEXT("Source")), TEXT("*.cpp"), true, false, false);
+		ISourceCodeAccessModule& SourceCodeAccessModule = FModuleManager::LoadModuleChecked<ISourceCodeAccessModule>("SourceCodeAccess");
+		if (OutProjectCodeFilenames.Num() > 0 && SourceCodeAccessModule.GetAccessor().CanAccessSourceCode())
+		{
+			InProfile->SetBuildGame(true);
+		}
+	}
 	FPlatformProcess::CreatePipe(ReadPipe, WritePipe);
 
 	// create task chains
diff --git a/Engine/Source/Programs/UnrealFrontend/Private/Commands/UserInterfaceCommand.cpp b/Engine/Source/Programs/UnrealFrontend/Private/Commands/UserInterfaceCommand.cpp
--- a/Engine/Source/Programs/UnrealFrontend/Private/Commands/UserInterfaceCommand.cpp
+++ b/Engine/Source/Programs/UnrealFrontend/Private/Commands/UserInterfaceCommand.cpp
@@ -8,6 +8,7 @@
 #include "StandaloneRenderer.h"
 #include "TaskGraphInterfaces.h"
 #include "UserInterfaceCommand.h"
+#include "ISourceCodeAccessModule.h"
 
 
 #define IDEAL_FRAMERATE 60;
@@ -47,6 +48,19 @@ void FUserInterfaceCommand::Run(  )
 
 	InitializeSlateApplication(UnrealFrontendLayoutIni);
 
+	// initialize source code access
+	// Load the source code access module
+	ISourceCodeAccessModule& SourceCodeAccessModule = FModuleManager::LoadModuleChecked<ISourceCodeAccessModule>( FName( "SourceCodeAccess" ) );
+
+	// Manually load in the source code access plugins, as standalone programs don't currently support plugins.
+#if PLATFORM_MAC
+	IModuleInterface& XCodeSourceCodeAccessModule = FModuleManager::LoadModuleChecked<IModuleInterface>( FName( "XCodeSourceCodeAccess" ) );
+	SourceCodeAccessModule.SetAccessor(FName("XCodeSourceCodeAccess"));
+#elif PLATFORM_WINDOWS
+	IModuleInterface& VisualStudioSourceCodeAccessModule = FModuleManager::LoadModuleChecked<IModuleInterface>( FName( "VisualStudioSourceCodeAccess" ) );
+	SourceCodeAccessModule.SetAccessor(FName("VisualStudioSourceCodeAccess"));
+#endif
+
 	// enter main loop
 	double DeltaTime = 0.0;
 	double LastTime = FPlatformTime::Seconds();
diff --git a/Engine/Source/Programs/UnrealFrontend/UnrealFrontend.Build.cs b/Engine/Source/Programs/UnrealFrontend/UnrealFrontend.Build.cs
--- a/Engine/Source/Programs/UnrealFrontend/UnrealFrontend.Build.cs
+++ b/Engine/Source/Programs/UnrealFrontend/UnrealFrontend.Build.cs
@@ -32,12 +32,22 @@ public class UnrealFrontend : ModuleRules
 				"Slate",
 				"SlateCore",
 				"SlateReflector",
+				"SourceCodeAccess",
 				"StandaloneRenderer",
 				"TargetDeviceServices",
 				"TargetPlatform",
 			}
 		);
 
+		if (Target.Platform == UnrealTargetPlatform.Mac)
+		{
+			PrivateDependencyModuleNames.Add("XCodeSourceCodeAccess");
+		}
+		else if (Target.Platform == UnrealTargetPlatform.Win64)
+		{
+			PrivateDependencyModuleNames.Add("VisualStudioSourceCodeAccess");
+		}
+
 		// @todo: allow for better plug-in support in standalone Slate apps
 		PrivateDependencyModuleNames.AddRange(
 			new string[] {
