removed an unneeded comment

diff --git a/Engine/Source/Developer/LauncherServices/Private/Launcher/LauncherWorker.cpp b/Engine/Source/Developer/LauncherServices/Private/Launcher/LauncherWorker.cpp
--- a/Engine/Source/Developer/LauncherServices/Private/Launcher/LauncherWorker.cpp
+++ b/Engine/Source/Developer/LauncherServices/Private/Launcher/LauncherWorker.cpp
@@ -533,7 +533,6 @@ void FLauncherWorker::CreateAndExecuteTasks( const ILauncherProfileRef& InProfil
 	// check to see if we need to build by default
 	if (!InProfile->HasProjectSpecified())
 	{
-		// TODO: make SouceCodeNavigation and GameProjectGeneration modules available to UFE
 		FString ProjectPath = FPaths::GetPath(InProfile->GetProjectPath());
 		TArray<FString> OutProjectCodeFilenames;
 		IFileManager::Get().FindFilesRecursive(OutProjectCodeFilenames, *(ProjectPath / TEXT("Source")), TEXT("*.h"), true, false, false);
