we shouldn't use incremental deploy when using a file server
#uat

diff --git a/Engine/Source/Developer/LauncherServices/Private/Launcher/LauncherWorker.cpp b/Engine/Source/Developer/LauncherServices/Private/Launcher/LauncherWorker.cpp
--- a/Engine/Source/Developer/LauncherServices/Private/Launcher/LauncherWorker.cpp
+++ b/Engine/Source/Developer/LauncherServices/Private/Launcher/LauncherWorker.cpp
@@ -448,14 +448,16 @@ FString FLauncherWorker::CreateUATCommand( const ILauncherProfileRef& InProfile,
 			}
 			break;
 
-		case ELauncherProfileDeploymentModes::FileServer:
 		case ELauncherProfileDeploymentModes::CopyToDevice:
 			{
-				UATCommand += TEXT(" -stage -deploy");
 				if (Profile->IsDeployingIncrementally())
 				{
 					UATCommand += " -iterativedeploy";
 				}
+			}
+		case ELauncherProfileDeploymentModes::FileServer:
+			{
+				UATCommand += TEXT(" -stage -deploy");
 				UATCommand += CommandLine;
 				UATCommand += StageDirectory;
 				UATCommand += DeviceCommand;
