Added default for create release version in the launcher profile.

diff --git a/Engine/Source/Developer/LauncherServices/Private/Profiles/LauncherProfile.h b/Engine/Source/Developer/LauncherServices/Private/Profiles/LauncherProfile.h
--- a/Engine/Source/Developer/LauncherServices/Private/Profiles/LauncherProfile.h
+++ b/Engine/Source/Developer/LauncherServices/Private/Profiles/LauncherProfile.h
@@ -722,7 +722,7 @@ public:
 		HideFileServerWindow = false;
 		DeployIncremental = false;
 
-		
+		CreateReleaseVersion = false;
 		GeneratePatch = false;
 		CreateDLC = false;
 
