Enabled MacServer cooking/packaging in UFE

diff --git a/Engine/Source/Developer/LauncherServices/Private/Launcher/LauncherWorker.cpp b/Engine/Source/Developer/LauncherServices/Private/Launcher/LauncherWorker.cpp
--- a/Engine/Source/Developer/LauncherServices/Private/Launcher/LauncherWorker.cpp
+++ b/Engine/Source/Developer/LauncherServices/Private/Launcher/LauncherWorker.cpp
@@ -219,6 +219,10 @@ FString FLauncherWorker::CreateUATCommand( const ILauncherProfileRef& InProfile,
 		{
 			ServerPlatforms += TEXT("+Win64");
 		}
+		else if (PlatformInfo->TargetPlatformName == FName("MacServer"))
+		{
+			ServerPlatforms += TEXT("+Mac");
+		}
 		else if (PlatformInfo->TargetPlatformName == FName("LinuxNoEditor"))
 		{
 			Platforms += TEXT("+Linux");
