Incrimented LauncherServices serialization version.

diff --git a/Engine/Source/Developer/LauncherServices/Private/LauncherServicesPrivatePCH.h b/Engine/Source/Developer/LauncherServices/Private/LauncherServicesPrivatePCH.h
--- a/Engine/Source/Developer/LauncherServices/Private/LauncherServicesPrivatePCH.h
+++ b/Engine/Source/Developer/LauncherServices/Private/LauncherServicesPrivatePCH.h
@@ -24,7 +24,7 @@
 #define LAUNCHERSERVICES_SIMPLEPROFILEVERSION 1
 
 /** Defines the launcher profile file version. */
-#define LAUNCHERSERVICES_PROFILEVERSION 11
+#define LAUNCHERSERVICES_PROFILEVERSION 14
 #define LAUNCHERSERVICES_MINPROFILEVERSION 10
 
 
