Compression on by default in the profile launcher.

diff --git a/Engine/Source/Developer/LauncherServices/Private/Profiles/LauncherProfile.h b/Engine/Source/Developer/LauncherServices/Private/Profiles/LauncherProfile.h
--- a/Engine/Source/Developer/LauncherServices/Private/Profiles/LauncherProfile.h
+++ b/Engine/Source/Developer/LauncherServices/Private/Profiles/LauncherProfile.h
@@ -757,7 +757,7 @@ public:
 		CookOptions = FString();
 		CookIncremental = false;
 		CookUnversioned = true;
-		Compressed = false;
+		Compressed = true;
 		CookedCultures.Reset();
 		CookedCultures.Add(I18N.GetCurrentCulture()->GetName());
 		CookedMaps.Reset();
