Added unversioned error message.
Fixed up error so it doesn't occur when using cook on the fly.

diff --git a/Engine/Source/Developer/LauncherServices/Private/Profiles/LauncherProfile.h b/Engine/Source/Developer/LauncherServices/Private/Profiles/LauncherProfile.h
--- a/Engine/Source/Developer/LauncherServices/Private/Profiles/LauncherProfile.h
+++ b/Engine/Source/Developer/LauncherServices/Private/Profiles/LauncherProfile.h
@@ -1224,7 +1224,7 @@ protected:
 			}
 		}
 
-		if ( CookUnversioned && CookIncremental )
+		if (CookUnversioned && CookIncremental && ((CookMode == ELauncherProfileCookModes::ByTheBook) || (CookMode == ELauncherProfileCookModes::ByTheBookInEditor)))
 		{
 			ValidationErrors.Add(ELauncherProfileValidationErrors::UnversionedAndIncrimental);
 		}
