FLauncherWorker now runs UAT in UTF-8 mode.

diff --git a/Engine/Source/Developer/LauncherServices/Private/Launcher/LauncherWorker.cpp b/Engine/Source/Developer/LauncherServices/Private/Launcher/LauncherWorker.cpp
--- a/Engine/Source/Developer/LauncherServices/Private/Launcher/LauncherWorker.cpp
+++ b/Engine/Source/Developer/LauncherServices/Private/Launcher/LauncherWorker.cpp
@@ -182,7 +182,7 @@ void FLauncherWorker::OnTaskCompleted(const FString& TaskName)
 
 FString FLauncherWorker::CreateUATCommand( const ILauncherProfileRef& InProfile, const TArray<FString>& InPlatforms, TArray<FCommandDesc>& OutCommands, FString& CommandStart )
 {
-	CommandStart = TEXT("");
+	CommandStart = TEXT(" -utf8output");
 	FString UATCommand = TEXT("");
 	static FGuid SessionId(FGuid::NewGuid());
 	FString InitialMap = InProfile->GetDefaultLaunchRole()->GetInitialMap();
