Put the -UTF8Output command in the right string.

diff --git a/Engine/Source/Developer/LauncherServices/Private/Launcher/LauncherWorker.cpp b/Engine/Source/Developer/LauncherServices/Private/Launcher/LauncherWorker.cpp
--- a/Engine/Source/Developer/LauncherServices/Private/Launcher/LauncherWorker.cpp
+++ b/Engine/Source/Developer/LauncherServices/Private/Launcher/LauncherWorker.cpp
@@ -182,8 +182,8 @@ void FLauncherWorker::OnTaskCompleted(const FString& TaskName)
 
 FString FLauncherWorker::CreateUATCommand( const ILauncherProfileRef& InProfile, const TArray<FString>& InPlatforms, TArray<FCommandDesc>& OutCommands, FString& CommandStart )
 {
-	CommandStart = TEXT(" -utf8output");
-	FString UATCommand = TEXT("");
+	CommandStart = TEXT("");
+	FString UATCommand = TEXT(" -utf8output");
 	static FGuid SessionId(FGuid::NewGuid());
 	FString InitialMap = InProfile->GetDefaultLaunchRole()->GetInitialMap();
 	if (InitialMap.IsEmpty() && InProfile->GetCookedMaps().Num() == 1)
