Default is to cook editor content, this is to avoid issue with slate requiring editor only content.

diff --git a/Engine/Source/Developer/LauncherServices/Private/Profiles/LauncherProfile.h b/Engine/Source/Developer/LauncherServices/Private/Profiles/LauncherProfile.h
--- a/Engine/Source/Developer/LauncherServices/Private/Profiles/LauncherProfile.h
+++ b/Engine/Source/Developer/LauncherServices/Private/Profiles/LauncherProfile.h
@@ -760,7 +760,7 @@ public:
 		CookedCultures.Add(I18N.GetCurrentCulture()->GetName());
 		CookedMaps.Reset();
 		CookedPlatforms.Reset();
-		bSkipCookingEditorContent = true;
+		bSkipCookingEditorContent = false;
         ForceClose = true;
         Timeout = 60;
 		NumCookersToSpawn = 0;
