#UE-18468 - LauncherUATTask properly passes -nocompileeditor when launching from editor or an installed engine situation. Remove a legacy static that was controlling NoCompile.

diff --git a/Engine/Source/Developer/LauncherServices/Private/Launcher/LauncherUATTask.h b/Engine/Source/Developer/LauncherServices/Private/Launcher/LauncherUATTask.h
--- a/Engine/Source/Developer/LauncherServices/Private/Launcher/LauncherUATTask.h
+++ b/Engine/Source/Developer/LauncherServices/Private/Launcher/LauncherUATTask.h
@@ -22,12 +22,10 @@ public:
 		, ProcessHandle(InProcessHandle)
 		, CommandText(InCommandEnd)
 	{
-		NoCompile = TEXT(" -nocompile");
 		EndTextFound = false;
 		InWorker->OnOutputReceived().AddRaw(this, &FLauncherUATTask::HandleOutputReceived);
 	}
 
-	static bool FirstTimeCompile;
 protected:
 
 	/**
@@ -39,12 +37,6 @@ protected:
 	 */
 	virtual bool PerformTask( FLauncherTaskChainState& ChainState ) override
 	{
-		if (FirstTimeCompile)
-		{
-			NoCompile = (!FParse::Param( FCommandLine::Get(), TEXT("development") ) && !ChainState.Profile->IsBuildingUAT()) ? TEXT(" -nocompile") : TEXT("");
-			FirstTimeCompile = false;
-		}
-
 		// spawn a UAT process to cook the data
 		// UAT executable
 		FString ExecutablePath = FPaths::ConvertRelativePathToFull(FPaths::EngineDir() + FString(TEXT("Build")) / TEXT("BatchFiles"));
@@ -72,7 +64,11 @@ protected:
 			*ConfigStrings[ChainState.Profile->GetBuildConfiguration()],
 			*ConfigStrings[ChainState.Profile->GetBuildConfiguration()]);
 
-		UATCommandLine += NoCompile;
+		// we expect to pass -nocompile to UAT here as we generally expect UAT to be fully compiled. Besides, installed builds don't even have the source to compile UAT scripts.
+		// Only allow UAT to compile scripts dynamically if we pass -development or we have the IsBuildingUAT property set, the latter of which should not allow itself to be set in installed situations.
+		UATCommandLine += FParse::Param( FCommandLine::Get(), TEXT("development") ) || ChainState.Profile->IsBuildingUAT() ? TEXT("") : TEXT(" -nocompile");
+		// we never want to build the editor when launching from the editor or running with an installed engine (which can't rebuild itself)
+		UATCommandLine += GIsEditor || FApp::IsEngineInstalled() ? TEXT(" -nocompileeditor") : TEXT("");
         UATCommandLine += Rocket;
 
 		// specify the path to the editor exe if necessary
@@ -117,9 +113,6 @@ private:
 	// Command line
 	FString CommandLine;
 
-	// Holds the no compile flag
-	FString NoCompile;
-
 	// The editor executable that UAT should use
 	FString EditorExe;
 
diff --git a/Engine/Source/Developer/LauncherServices/Private/Launcher/LauncherWorker.cpp b/Engine/Source/Developer/LauncherServices/Private/Launcher/LauncherWorker.cpp
--- a/Engine/Source/Developer/LauncherServices/Private/Launcher/LauncherWorker.cpp
+++ b/Engine/Source/Developer/LauncherServices/Private/Launcher/LauncherWorker.cpp
@@ -6,8 +6,6 @@
 
 #define LOCTEXT_NAMESPACE "LauncherWorker"
 
-bool FLauncherUATTask::FirstTimeCompile = true;
-
 /* FLauncherWorker structors
  *****************************************************************************/
 
@@ -616,8 +614,6 @@ void FLauncherWorker::CreateAndExecuteTasks( const ILauncherProfileRef& InProfil
 		Platforms = InProfile->GetCookedPlatforms();
 	}
 
-	FLauncherUATTask::FirstTimeCompile = true;
-
 	// determine deployment platforms
 	ILauncherDeviceGroupPtr DeviceGroup = InProfile->GetDeployedDeviceGroup();
 	FName Variant = NAME_None;
