Move UFE launch profiles from Engine/Saved to Engine/Programs/UnrealFrontend/Profiles.
Change the filenames to be more readable.

diff --git a/Engine/Source/Developer/LauncherServices/Private/Profiles/LauncherProfile.h b/Engine/Source/Developer/LauncherServices/Private/Profiles/LauncherProfile.h
--- a/Engine/Source/Developer/LauncherServices/Private/Profiles/LauncherProfile.h
+++ b/Engine/Source/Developer/LauncherServices/Private/Profiles/LauncherProfile.h
@@ -153,7 +153,14 @@ public:
 	/**
 	 * Destructor.
 	 */
-	~FLauncherProfile( ) { }
+	~FLauncherProfile( ) 
+	{
+		if (DeployedDeviceGroup.IsValid())
+		{
+			DeployedDeviceGroup->OnDeviceAdded().Remove(OnLauncherDeviceGroupDeviceAddedDelegateHandle);
+			DeployedDeviceGroup->OnDeviceRemoved().Remove(OnLauncherDeviceGroupDeviceRemoveDelegateHandle);
+		}
+	}
 
 public:
 
@@ -375,6 +382,11 @@ public:
 		return Id;
 	}
 
+	virtual FString GetFileName() const override
+	{
+		return GetName() + TEXT("_") + GetId().ToString() + TEXT(".ulp");
+	}
+
 	virtual ELauncherProfileLaunchModes::Type GetLaunchMode( ) const override
 	{
 		return LaunchMode;
diff --git a/Engine/Source/Developer/LauncherServices/Private/Profiles/LauncherProfileManager.cpp b/Engine/Source/Developer/LauncherServices/Private/Profiles/LauncherProfileManager.cpp
--- a/Engine/Source/Developer/LauncherServices/Private/Profiles/LauncherProfileManager.cpp
+++ b/Engine/Source/Developer/LauncherServices/Private/Profiles/LauncherProfileManager.cpp
@@ -267,7 +267,7 @@ void FLauncherProfileManager::RemoveProfile( const ILauncherProfileRef& Profile
 		if (Profile->GetId().IsValid())
 		{
 			// delete the persisted profile on disk
-			FString ProfileFileName = GetProfileFolder() / Profile->GetId().ToString() + TEXT(".ulp");
+			FString ProfileFileName = GetProfileFolder() / Profile->GetFileName();
 
 			// delete the profile
 			IFileManager::Get().Delete(*ProfileFileName);
@@ -278,11 +278,11 @@ void FLauncherProfileManager::RemoveProfile( const ILauncherProfileRef& Profile
 }
 
 
-void FLauncherProfileManager::SaveProfile(const ILauncherProfileRef& Profile)
+bool FLauncherProfileManager::SaveProfile(const ILauncherProfileRef& Profile)
 {
 	if (Profile->GetId().IsValid())
 	{
-		FString ProfileFileName = GetProfileFolder() / Profile->GetId().ToString() + TEXT(".ulp");
+		FString ProfileFileName = GetProfileFolder() / Profile->GetFileName();
 		FArchive* ProfileFileWriter = IFileManager::Get().CreateFileWriter(*ProfileFileName);
 
 		if (ProfileFileWriter != nullptr)
@@ -290,10 +290,33 @@ void FLauncherProfileManager::SaveProfile(const ILauncherProfileRef& Profile)
 			Profile->Serialize(*ProfileFileWriter);
 
 			delete ProfileFileWriter;
+
+			return true;
 		}
 	}
+	return false;
 }
 
+void FLauncherProfileManager::ChangeProfileName(const ILauncherProfileRef& Profile, FString Name)
+{
+	FString OldName = Profile->GetName();
+	FString OldProfileFileName = GetProfileFolder() / Profile->GetFileName();
+
+	//change name and save to new location
+	Profile->SetName(Name);
+	if (SaveProfile(Profile))
+	{
+		// delete the old profile
+		IFileManager::Get().Delete(*OldProfileFileName);
+	}
+	else
+	{
+		//if we couldn't save successfully, change the name back to keep files/profiles matching.
+		Profile->SetName(OldName);		
+	}	
+}
+
+
 
 void FLauncherProfileManager::SaveSettings( )
 {
@@ -357,6 +380,32 @@ void FLauncherProfileManager::LoadProfiles( )
 {
 	TArray<FString> ProfileFileNames;
 
+	//load and move legacy profiles
+	{
+		IFileManager::Get().FindFiles(ProfileFileNames, *(GetLegacyProfileFolder() / TEXT("*.ulp")), true, false);
+		for (TArray<FString>::TConstIterator It(ProfileFileNames); It; ++It)
+		{
+			FString ProfileFilePath = GetLegacyProfileFolder() / *It;
+			FArchive* ProfileFileReader = IFileManager::Get().CreateFileReader(*ProfileFilePath);
+
+			if (ProfileFileReader != nullptr)
+			{
+				ILauncherProfilePtr LoadedProfile = LoadProfile(*ProfileFileReader);
+				delete ProfileFileReader;
+
+				//resave profile to new location
+				if (LoadedProfile.IsValid())
+				{
+					SaveProfile(LoadedProfile.ToSharedRef());
+				}
+
+				//delete legacy profile.
+				IFileManager::Get().Delete(*ProfileFilePath);				
+			}
+		}
+	}
+
+	ProfileFileNames.Reset();
 	IFileManager::Get().FindFiles(ProfileFileNames, *(GetProfileFolder() / TEXT("*.ulp")), true, false);
 	
 	for (TArray<FString>::TConstIterator It(ProfileFileNames); It; ++It)
diff --git a/Engine/Source/Developer/LauncherServices/Private/Profiles/LauncherProfileManager.h b/Engine/Source/Developer/LauncherServices/Private/Profiles/LauncherProfileManager.h
--- a/Engine/Source/Developer/LauncherServices/Private/Profiles/LauncherProfileManager.h
+++ b/Engine/Source/Developer/LauncherServices/Private/Profiles/LauncherProfileManager.h
@@ -83,7 +83,9 @@ public:
 
 	virtual void RemoveProfile( const ILauncherProfileRef& Profile ) override;
 
-	virtual void SaveProfile( const ILauncherProfileRef& Profile) override;
+	virtual bool SaveProfile( const ILauncherProfileRef& Profile) override;
+
+	virtual void ChangeProfileName( const ILauncherProfileRef& Profile, FString Name) override;
 
 	virtual void SaveSettings( ) override;
 
@@ -130,6 +132,16 @@ protected:
 
 protected:
 
+	/**
+	* Gets the folder in which old profile files were stored.
+	*
+	* @return The folder path.
+	*/
+	static FString GetLegacyProfileFolder()
+	{
+		return FPaths::EngineSavedDir() / TEXT("Launcher");
+	}
+
 	/**
 	 * Gets the folder in which profile files are stored.
 	 *
@@ -137,7 +149,7 @@ protected:
 	 */
 	static FString GetProfileFolder( )
 	{
-		return FPaths::EngineSavedDir() / TEXT("Launcher");
+		return FPaths::EngineDir() / TEXT("Programs/UnrealFrontend/Profiles");
 	}
 
 private:
diff --git a/Engine/Source/Developer/LauncherServices/Public/Interfaces/ILauncherProfile.h b/Engine/Source/Developer/LauncherServices/Public/Interfaces/ILauncherProfile.h
--- a/Engine/Source/Developer/LauncherServices/Public/Interfaces/ILauncherProfile.h
+++ b/Engine/Source/Developer/LauncherServices/Public/Interfaces/ILauncherProfile.h
@@ -314,6 +314,13 @@ public:
 	 */
 	virtual FGuid GetId( ) const = 0;
 
+	/**
+	* Gets the file name for serialization.
+	*
+	* @return The file name.
+	*/
+	virtual FString GetFileName( ) const = 0;
+
 	/**
 	 * Gets the human readable name of the profile.
 	 *
diff --git a/Engine/Source/Developer/LauncherServices/Public/Interfaces/ILauncherProfileManager.h b/Engine/Source/Developer/LauncherServices/Public/Interfaces/ILauncherProfileManager.h
--- a/Engine/Source/Developer/LauncherServices/Public/Interfaces/ILauncherProfileManager.h
+++ b/Engine/Source/Developer/LauncherServices/Public/Interfaces/ILauncherProfileManager.h
@@ -197,7 +197,14 @@ public:
 	 * @param Archive The archive to save to.
 	 * @see LoadProfile
 	 */
-	virtual void SaveProfile(const ILauncherProfileRef& Profile) = 0;
+	virtual bool SaveProfile(const ILauncherProfileRef& Profile) = 0;
+
+	/**
+	* Modifies profile name. Saves file to new location and deletes old one to avoid duplicating profiles.
+	*
+	* @param Profile The profile to change	
+	*/
+	virtual void ChangeProfileName(const ILauncherProfileRef& Profile, FString Name) = 0;
 
 public:
 
diff --git a/Engine/Source/Developer/ProjectLauncher/Private/Widgets/SProjectLauncherSettings.cpp b/Engine/Source/Developer/ProjectLauncher/Private/Widgets/SProjectLauncherSettings.cpp
--- a/Engine/Source/Developer/ProjectLauncher/Private/Widgets/SProjectLauncherSettings.cpp
+++ b/Engine/Source/Developer/ProjectLauncher/Private/Widgets/SProjectLauncherSettings.cpp
@@ -299,7 +299,7 @@ void SProjectLauncherSettings::OnNameTextCommitted(const FText& NewText, ETextCo
 	const ILauncherProfilePtr& LaunchProfile = Model->GetSelectedProfile();
 	if (LaunchProfile.IsValid())
 	{
-		LaunchProfile->SetName(NewText.ToString());
+		Model->GetProfileManager()->ChangeProfileName(LaunchProfile.ToSharedRef(), NewText.ToString());
 	}
 }
 
