Let UFE profiles be searchable under NotForLicensees
Fix a renaming bug.
Fix for renaming to work properly on case-sensitive and case-sensitive platforms.

diff --git a/Engine/Source/Developer/LauncherServices/Private/Profiles/LauncherProfile.h b/Engine/Source/Developer/LauncherServices/Private/Profiles/LauncherProfile.h
--- a/Engine/Source/Developer/LauncherServices/Private/Profiles/LauncherProfile.h
+++ b/Engine/Source/Developer/LauncherServices/Private/Profiles/LauncherProfile.h
@@ -99,7 +99,7 @@ public:
 		// profiles created for your persistent devices to be in debug. The user might not see this if they don't expand the Advanced options.
 		BuildConfiguration = EBuildConfigurations::Development;
 		
-		CookMode = ELauncherProfileCookModes::OnTheFly;
+		CookMode = ELauncherProfileCookModes::OnTheFly;		
 	}
 
 private:
@@ -126,6 +126,16 @@ class FLauncherProfile
 {
 public:
 
+	/**
+	* Gets the folder in which profile files are stored.
+	*
+	* @return The folder path.
+	*/
+	static FString GetProfileFolder()
+	{
+		return FPaths::EngineDir() / TEXT("Programs/UnrealFrontend/Profiles");
+	}
+
 	/**
 	 * Default constructor.
 	 */
@@ -384,7 +394,17 @@ public:
 
 	virtual FString GetFileName() const override
 	{
-		return GetName() + TEXT("_") + GetId().ToString() + TEXT(".ulp");
+		//toupper for filename so that filepaths can be compared the same on case sensitive and case-insensitive platforms
+		return GetName().ToUpper() + TEXT("_") + GetId().ToString() + TEXT(".ulp");
+	}
+
+	virtual FString GetFilePath() const override
+	{
+		if (bNotForLicensees)
+		{
+			return GetProfileFolder() / "NotForLicensees" / GetFileName();			
+		}
+		return GetProfileFolder() / GetFileName();
 	}
 
 	virtual ELauncherProfileLaunchModes::Type GetLaunchMode( ) const override
@@ -817,6 +837,8 @@ public:
 		// default UAT settings
 		EditorExe.Empty();
 
+		bNotForLicensees = false;
+
 		Validate();
 	}
 
@@ -1060,6 +1082,11 @@ public:
 		}
 	}
 
+	virtual void SetNotForLicensees() override
+	{
+		bNotForLicensees = true;
+	}
+
 	virtual void SetPackagingMode( ELauncherProfilePackagingModes::Type Mode ) override
 	{
 		if (PackagingMode != Mode)
@@ -1522,6 +1549,9 @@ private:
 	// Path to the editor executable to pass to UAT, for cooking, etc... May be empty.
 	FString EditorExe;
 
+	// Profile is for an internal project
+	bool bNotForLicensees;
+
 private:
 
 	// cook in the editor callbacks (not valid for any other cook mode)
diff --git a/Engine/Source/Developer/LauncherServices/Private/Profiles/LauncherProfileManager.cpp b/Engine/Source/Developer/LauncherServices/Private/Profiles/LauncherProfileManager.cpp
--- a/Engine/Source/Developer/LauncherServices/Private/Profiles/LauncherProfileManager.cpp
+++ b/Engine/Source/Developer/LauncherServices/Private/Profiles/LauncherProfileManager.cpp
@@ -253,7 +253,7 @@ void FLauncherProfileManager::RemoveSimpleProfile(const ILauncherSimpleProfileRe
 	if (SimpleProfiles.Remove(SimpleProfile) > 0)
 	{
 		// delete the persisted simple profile on disk
-		FString SimpleProfileFileName = GetProfileFolder() / SimpleProfile->GetDeviceName() + TEXT(".uslp");
+		FString SimpleProfileFileName = FLauncherProfile::GetProfileFolder() / SimpleProfile->GetDeviceName() + TEXT(".uslp");
 		IFileManager::Get().Delete(*SimpleProfileFileName);
 	}
 }
@@ -267,7 +267,7 @@ void FLauncherProfileManager::RemoveProfile( const ILauncherProfileRef& Profile
 		if (Profile->GetId().IsValid())
 		{
 			// delete the persisted profile on disk
-			FString ProfileFileName = GetProfileFolder() / Profile->GetFileName();
+			FString ProfileFileName = Profile->GetFilePath();
 
 			// delete the profile
 			IFileManager::Get().Delete(*ProfileFileName);
@@ -282,7 +282,7 @@ bool FLauncherProfileManager::SaveProfile(const ILauncherProfileRef& Profile)
 {
 	if (Profile->GetId().IsValid())
 	{
-		FString ProfileFileName = GetProfileFolder() / Profile->GetFileName();
+		FString ProfileFileName = Profile->GetFilePath();
 		FArchive* ProfileFileWriter = IFileManager::Get().CreateFileWriter(*ProfileFileName);
 
 		if (ProfileFileWriter != nullptr)
@@ -300,19 +300,23 @@ bool FLauncherProfileManager::SaveProfile(const ILauncherProfileRef& Profile)
 void FLauncherProfileManager::ChangeProfileName(const ILauncherProfileRef& Profile, FString Name)
 {
 	FString OldName = Profile->GetName();
-	FString OldProfileFileName = GetProfileFolder() / Profile->GetFileName();
+	FString OldProfileFileName = Profile->GetFilePath();
 
 	//change name and save to new location
 	Profile->SetName(Name);
 	if (SaveProfile(Profile))
 	{
-		// delete the old profile
-		IFileManager::Get().Delete(*OldProfileFileName);
+		//delete the old profile if the location moved.  File names should be uppercase so this compare works on case sensitive and insensitive platforms
+		if (OldProfileFileName.Compare(Profile->GetFilePath()) != 0)
+		{
+			
+			IFileManager::Get().Delete(*OldProfileFileName);
+		}
 	}
 	else
 	{
 		//if we couldn't save successfully, change the name back to keep files/profiles matching.
-		Profile->SetName(OldName);		
+		Profile->SetName(OldName);
 	}	
 }
 
@@ -406,11 +410,11 @@ void FLauncherProfileManager::LoadProfiles( )
 	}
 
 	ProfileFileNames.Reset();
-	IFileManager::Get().FindFiles(ProfileFileNames, *(GetProfileFolder() / TEXT("*.ulp")), true, false);
+	IFileManager::Get().FindFilesRecursive(ProfileFileNames, *FLauncherProfile::GetProfileFolder(), TEXT("*.ulp"), true, false);
 	
 	for (TArray<FString>::TConstIterator It(ProfileFileNames); It; ++It)
 	{
-		FString ProfileFilePath = GetProfileFolder() / *It;
+		FString ProfileFilePath = *It;
 		FArchive* ProfileFileReader = IFileManager::Get().CreateFileReader(*ProfileFilePath);
 
 		if (ProfileFileReader != nullptr)
@@ -420,6 +424,10 @@ void FLauncherProfileManager::LoadProfiles( )
 
 			if (LoadedProfile.IsValid())
 			{
+				if (ProfileFilePath.Contains("NotForLicensees"))
+				{
+					LoadedProfile->SetNotForLicensees();
+				}
 				AddProfile(LoadedProfile.ToSharedRef());
 			}
 			else
@@ -509,7 +517,7 @@ void FLauncherProfileManager::SaveSimpleProfiles()
 {
 	for (TArray<ILauncherSimpleProfilePtr>::TIterator It(SimpleProfiles); It; ++It)
 	{
-		FString SimpleProfileFileName = GetProfileFolder() / (*It)->GetDeviceName() + TEXT(".uslp");
+		FString SimpleProfileFileName = FLauncherProfile::GetProfileFolder() / (*It)->GetDeviceName() + TEXT(".uslp");
 		FArchive* ProfileFileWriter = IFileManager::Get().CreateFileWriter(*SimpleProfileFileName);
 
 		if (ProfileFileWriter != nullptr)
diff --git a/Engine/Source/Developer/LauncherServices/Private/Profiles/LauncherProfileManager.h b/Engine/Source/Developer/LauncherServices/Private/Profiles/LauncherProfileManager.h
--- a/Engine/Source/Developer/LauncherServices/Private/Profiles/LauncherProfileManager.h
+++ b/Engine/Source/Developer/LauncherServices/Private/Profiles/LauncherProfileManager.h
@@ -142,16 +142,6 @@ protected:
 		return FPaths::EngineSavedDir() / TEXT("Launcher");
 	}
 
-	/**
-	 * Gets the folder in which profile files are stored.
-	 *
-	 * @return The folder path.
-	 */
-	static FString GetProfileFolder( )
-	{
-		return FPaths::EngineDir() / TEXT("Programs/UnrealFrontend/Profiles");
-	}
-
 private:
 
 	// Holds the collection of device groups.
diff --git a/Engine/Source/Developer/LauncherServices/Public/Interfaces/ILauncherProfile.h b/Engine/Source/Developer/LauncherServices/Public/Interfaces/ILauncherProfile.h
--- a/Engine/Source/Developer/LauncherServices/Public/Interfaces/ILauncherProfile.h
+++ b/Engine/Source/Developer/LauncherServices/Public/Interfaces/ILauncherProfile.h
@@ -321,6 +321,13 @@ public:
 	*/
 	virtual FString GetFileName( ) const = 0;
 
+	/**
+	* Gets the full file path for serialization.
+	*
+	* @return The file path.
+	*/
+	virtual FString GetFilePath() const = 0;
+
 	/**
 	 * Gets the human readable name of the profile.
 	 *
@@ -400,6 +407,12 @@ public:
 	 */
 	virtual void SetDescription(const FString& NewDescription) = 0;
 
+	/**
+	* Changes the save location to an internal project path.
+	*	
+	*/
+	virtual void SetNotForLicensees() = 0;
+
 	/**
 	 * Returns the cook delegate which can be used to query if the cook is finished.
 	 *
diff --git a/Engine/Source/Developer/ProjectLauncher/Private/Widgets/Profile/SProjectLauncherProfileListRow.h b/Engine/Source/Developer/ProjectLauncher/Private/Widgets/Profile/SProjectLauncherProfileListRow.h
--- a/Engine/Source/Developer/ProjectLauncher/Private/Widgets/Profile/SProjectLauncherProfileListRow.h
+++ b/Engine/Source/Developer/ProjectLauncher/Private/Widgets/Profile/SProjectLauncherProfileListRow.h
@@ -73,7 +73,7 @@ public:
 					.FillWidth(1.0f)
 					.VAlign(VAlign_Center)
 					[
-						SAssignNew(ProfileNameDescEditor, SProjectLauncherProfileNameDescEditor, false)
+						SAssignNew(ProfileNameDescEditor, SProjectLauncherProfileNameDescEditor, InModel, false)
 						.LaunchProfile(this, &SProjectLauncherProfileListRow::GetLaunchProfile)
 					]
 
diff --git a/Engine/Source/Developer/ProjectLauncher/Private/Widgets/SProjectLauncherSettings.cpp b/Engine/Source/Developer/ProjectLauncher/Private/Widgets/SProjectLauncherSettings.cpp
--- a/Engine/Source/Developer/ProjectLauncher/Private/Widgets/SProjectLauncherSettings.cpp
+++ b/Engine/Source/Developer/ProjectLauncher/Private/Widgets/SProjectLauncherSettings.cpp
@@ -47,7 +47,7 @@ void SProjectLauncherSettings::Construct( const FArguments& InArgs, const FProje
 				.VAlign(VAlign_Center)
 				.Padding(4.0f, 0.0f, 4.0f, 0.0f)
 				[
-					SNew(SProjectLauncherProfileNameDescEditor, true)
+					SNew(SProjectLauncherProfileNameDescEditor, InModel, true)
 					.LaunchProfile(this, &SProjectLauncherSettings::GetLaunchProfile)
 				]
 
diff --git a/Engine/Source/Developer/ProjectLauncher/Private/Widgets/Shared/SProjectLauncherProfileNameDescEditor.cpp b/Engine/Source/Developer/ProjectLauncher/Private/Widgets/Shared/SProjectLauncherProfileNameDescEditor.cpp
--- a/Engine/Source/Developer/ProjectLauncher/Private/Widgets/Shared/SProjectLauncherProfileNameDescEditor.cpp
+++ b/Engine/Source/Developer/ProjectLauncher/Private/Widgets/Shared/SProjectLauncherProfileNameDescEditor.cpp
@@ -9,10 +9,11 @@
 
 
 
-void SProjectLauncherProfileNameDescEditor::Construct(const FArguments& InArgs, bool InShowAddDescriptionText)
+void SProjectLauncherProfileNameDescEditor::Construct(const FArguments& InArgs, const FProjectLauncherModelRef& InModel, bool InShowAddDescriptionText)
 {
 	EnterTextDescription = FText(LOCTEXT("LaunchProfileEnterDescription", "Enter a description here."));
 
+	Model = InModel;
 	LaunchProfileAttr = InArgs._LaunchProfile;
 	bShowAddDescriptionText = InShowAddDescriptionText;
 
diff --git a/Engine/Source/Developer/ProjectLauncher/Private/Widgets/Shared/SProjectLauncherProfileNameDescEditor.h b/Engine/Source/Developer/ProjectLauncher/Private/Widgets/Shared/SProjectLauncherProfileNameDescEditor.h
--- a/Engine/Source/Developer/ProjectLauncher/Private/Widgets/Shared/SProjectLauncherProfileNameDescEditor.h
+++ b/Engine/Source/Developer/ProjectLauncher/Private/Widgets/Shared/SProjectLauncherProfileNameDescEditor.h
@@ -13,6 +13,7 @@ public:
 
 	SLATE_BEGIN_ARGS(SProjectLauncherProfileNameDescEditor) { }
 		SLATE_ATTRIBUTE(ILauncherProfilePtr, LaunchProfile)
+		SLATE_ATTRIBUTE(ILauncherProfileManagerPtr, InModel)
 	SLATE_END_ARGS()
 
 public:
@@ -23,7 +24,7 @@ public:
 	 * @param InArgs The Slate argument list.
 	 * @param InModel The data model.
 	 */
-	void Construct(const FArguments& InArgs, bool InShowAddDescriptionText);
+	void Construct(const FArguments& InArgs, const FProjectLauncherModelRef& InModel, bool InShowAddDescriptionText);
 
 	/**
 	 * Triggers a name edit for the profile.
@@ -57,7 +58,7 @@ private:
 		ILauncherProfilePtr LaunchProfile = LaunchProfileAttr.Get();
 		if (LaunchProfile.IsValid())
 		{
-			LaunchProfile->SetName(NewText.ToString());
+			Model->GetProfileManager()->ChangeProfileName(LaunchProfile.ToSharedRef(), NewText.ToString());
 		}
 	}
 
@@ -95,6 +96,9 @@ private:
 
 private:
 
+	/** Holds a pointer to the data model. */
+	FProjectLauncherModelPtr Model;
+
 	// Attribute for the launch profile this widget edits. 
 	TAttribute<ILauncherProfilePtr> LaunchProfileAttr;
 
