Add default platform deploy to UFE for UFE profile sharing across teams.
Fix some bugs with the deploy screen getting out of sync with the actual profile data.

diff --git a/Engine/Source/Developer/LauncherServices/Private/Launcher/LauncherWorker.cpp b/Engine/Source/Developer/LauncherServices/Private/Launcher/LauncherWorker.cpp
--- a/Engine/Source/Developer/LauncherServices/Private/Launcher/LauncherWorker.cpp
+++ b/Engine/Source/Developer/LauncherServices/Private/Launcher/LauncherWorker.cpp
@@ -177,6 +177,24 @@ void FLauncherWorker::OnTaskCompleted(const FString& TaskName)
 	StageCompleted.Broadcast(TaskName, FPlatformTime::Seconds() - StageStartTime);
 }
 
+static void AddDeviceToLaunchCommand(const FString& DeviceId, ITargetDeviceProxyPtr DeviceProxy, const ILauncherProfileRef& InProfile, FString& DeviceNames, FString& RoleCommands, bool& bVsyncAdded)
+{
+	// add the platform
+	DeviceNames += TEXT("+\"") + DeviceId + TEXT("\"");
+	TArray<ILauncherProfileLaunchRolePtr> Roles;
+	if (InProfile->GetLaunchRolesFor(DeviceId, Roles) > 0)
+	{
+		for (int32 RoleIndex = 0; RoleIndex < Roles.Num(); RoleIndex++)
+		{
+			if (!bVsyncAdded && Roles[RoleIndex]->IsVsyncEnabled())
+			{
+				RoleCommands += TEXT(" -vsync");
+				bVsyncAdded = true;
+			}
+			RoleCommands += *(TEXT(" ") + Roles[RoleIndex]->GetCommandLine());
+		}
+	}
+}
 
 FString FLauncherWorker::CreateUATCommand( const ILauncherProfileRef& InProfile, const TArray<FString>& InPlatforms, TArray<FCommandDesc>& OutCommands, FString& CommandStart )
 {
@@ -265,38 +283,25 @@ FString FLauncherWorker::CreateUATCommand( const ILauncherProfileRef& InProfile,
 	FString DeviceCommand = TEXT("");
 	FString RoleCommands = TEXT("");
 	ILauncherDeviceGroupPtr DeviceGroup = InProfile->GetDeployedDeviceGroup();
+
+	bool bVsyncAdded = false;
+
 	if (DeviceGroup.IsValid())
 	{
-		const TArray<FString>& Devices = DeviceGroup->GetDeviceIDs();
-		bool bVsyncAdded = false;
+		const TArray<FString>& Devices = DeviceGroup->GetDeviceIDs();		
 
 		// for each deployed device...
 		for (int32 DeviceIndex = 0; DeviceIndex < Devices.Num(); ++DeviceIndex)
 		{
 			const FString& DeviceId = Devices[DeviceIndex];
-
 			ITargetDeviceProxyPtr DeviceProxy = DeviceProxyManager->FindProxyDeviceForTargetDevice(DeviceId);
-
 			if (DeviceProxy.IsValid())
 			{
-				// add the platform
-				DeviceNames += TEXT("+\"") + DeviceId + TEXT("\"");
-				TArray<ILauncherProfileLaunchRolePtr> Roles;
-				if (InProfile->GetLaunchRolesFor(DeviceId, Roles) > 0)
-				{
-					for (int32 RoleIndex = 0; RoleIndex < Roles.Num(); RoleIndex++)
-					{
-						if (!bVsyncAdded && Roles[RoleIndex]->IsVsyncEnabled())
-						{
-							RoleCommands += TEXT(" -vsync");
-							bVsyncAdded = true;
-						}
-						RoleCommands += *(TEXT(" ") + Roles[RoleIndex]->GetCommandLine());
-					}
-				}
+				AddDeviceToLaunchCommand(DeviceId, DeviceProxy, InProfile, DeviceNames, RoleCommands, bVsyncAdded);
 			}			
 		}
 	}
+
 	if (DeviceNames.Len() > 0)
 	{
 		DeviceCommand += TEXT(" -device=") + DeviceNames.RightChop(1);
diff --git a/Engine/Source/Developer/LauncherServices/Private/Profiles/LauncherDeviceGroup.h b/Engine/Source/Developer/LauncherServices/Private/Profiles/LauncherDeviceGroup.h
--- a/Engine/Source/Developer/LauncherServices/Private/Profiles/LauncherDeviceGroup.h
+++ b/Engine/Source/Developer/LauncherServices/Private/Profiles/LauncherDeviceGroup.h
@@ -76,6 +76,17 @@ public:
 		DeviceRemovedDelegate.Broadcast(AsShared(), DeviceID);
 	}
 
+	virtual void RemoveAllDevices() override
+	{
+		for (int32 i = Devices.Num() - 1; i >= 0; --i)
+		{
+			FString DeviceId = Devices[i];
+			Devices.RemoveAt(i, 1, false);
+
+			DeviceRemovedDelegate.Broadcast(AsShared(), DeviceId);
+		}
+	}
+
 	virtual void SetName(const FString& NewName) override
 	{
 		Name = NewName;
diff --git a/Engine/Source/Developer/LauncherServices/Private/Profiles/LauncherProfile.h b/Engine/Source/Developer/LauncherServices/Private/Profiles/LauncherProfile.h
--- a/Engine/Source/Developer/LauncherServices/Private/Profiles/LauncherProfile.h
+++ b/Engine/Source/Developer/LauncherServices/Private/Profiles/LauncherProfile.h
@@ -10,6 +10,7 @@
 #define LAUNCHERSERVICES_ADDEDGENERATECHUNKS 16
 #define LAUNCHERSERVICES_ADDEDNUMCOOKERSTOSPAWN 17
 #define LAUNCHERSERVICES_ADDEDSKIPCOOKINGEDITORCONTENT 18
+#define LAUNCHERSERVICES_ADDEDDEFAULTDEPLOYPLATFORM 19
 
 /**
 * Implements a simple profile which controls the desired output of the Launcher for simple
@@ -213,6 +214,45 @@ public:
 		Validate();
 	}
 
+	virtual void SetDefaultDeployPlatform(const FName PlatformName) override
+	{
+		
+		DefaultDeployPlatform = PlatformName;	
+
+		if (DeployedDeviceGroup.IsValid())
+		{
+			DeployedDeviceGroup->RemoveAllDevices();
+
+			if (DefaultDeployPlatform != NAME_None)
+			{
+				TArray<ITargetDeviceProxyPtr> PlatformDeviceProxies;
+				ITargetDeviceServicesModule& TargetDeviceServicesModule = FModuleManager::LoadModuleChecked<ITargetDeviceServicesModule>("TargetDeviceServices");
+				const ITargetDeviceProxyManagerRef& InDeviceProxyManager = TargetDeviceServicesModule.GetDeviceProxyManager();
+
+				InDeviceProxyManager->GetProxies(NAME_None, true, PlatformDeviceProxies);
+
+				ITargetDeviceProxyPtr DefaultPlatformDevice;
+				for (int32 ProxyIndex = 0; ProxyIndex < PlatformDeviceProxies.Num(); ++ProxyIndex)
+				{
+					ITargetDeviceProxyPtr DeviceProxy = PlatformDeviceProxies[ProxyIndex];
+
+					if (DeviceProxy->GetVanillaPlatformId(NAME_None) == DefaultDeployPlatform)
+					{
+						DefaultPlatformDevice = DeviceProxy;
+						break;
+					}
+				}
+
+				if (DefaultPlatformDevice.IsValid())
+				{
+					DeployedDeviceGroup->AddDevice(DefaultPlatformDevice->GetTargetDeviceId(NAME_None));
+				}
+			}
+		}
+
+		Validate();
+	}
+
 	virtual void ClearCookedCultures( ) override
 	{
 		if (CookedCultures.Num() > 0)
@@ -304,11 +344,22 @@ public:
 		return DefaultLaunchRole;
 	}
 
-	virtual ILauncherDeviceGroupPtr GetDeployedDeviceGroup( ) const override
+	virtual ILauncherDeviceGroupPtr GetDeployedDeviceGroup( ) override
 	{
+		// setting the default platform will update the device group.  always do this when getting the group because
+		// devices come in lazily through messages and can't be added properly at profile load.
+		if (DefaultDeployPlatform != NAME_None)
+		{
+			SetDefaultDeployPlatform(DefaultDeployPlatform);
+		}
 		return DeployedDeviceGroup;
 	}
 
+	virtual const FName GetDefaultDeployPlatform() const override
+	{
+		return DefaultDeployPlatform;
+	}
+
 	virtual bool IsGeneratingPatch() const override
 	{
 		return GeneratePatch;
@@ -628,7 +679,7 @@ public:
 
 	virtual bool Serialize( FArchive& Archive ) override
 	{
-		int32 Version = LAUNCHERSERVICES_ADDEDSKIPCOOKINGEDITORCONTENT;
+		int32 Version = LAUNCHERSERVICES_ADDEDDEFAULTDEPLOYPLATFORM;
 
 		Archive	<< Version;
 
@@ -676,6 +727,11 @@ public:
                 << ForceClose
                 << Timeout;
 
+		FString DeployPlatformString = DefaultDeployPlatform.ToString();
+		if (Version >= LAUNCHERSERVICES_ADDEDDEFAULTDEPLOYPLATFORM)
+		{
+			Archive << DeployPlatformString;
+		}		
 		if (Version >= LAUNCHERSERVICES_ADDEDNUMCOOKERSTOSPAWN)
 		{
 			Archive << NumCookersToSpawn;
@@ -745,6 +801,16 @@ public:
 			}
 		}
 
+		if (Archive.IsLoading())
+		{
+			DefaultDeployPlatform = FName(*DeployPlatformString);
+		}
+
+		if (DefaultDeployPlatform != NAME_None)
+		{
+			SetDefaultDeployPlatform(DefaultDeployPlatform);
+		}
+
 		Validate();
 
 		return true;
@@ -822,6 +888,7 @@ public:
 		HttpChunkDataReleaseName = TEXT("");
 
 		// default launch settings
+		DefaultDeployPlatform = NAME_None;
 		LaunchMode = ELauncherProfileLaunchModes::DefaultRole;
 		DefaultLaunchRole->SetCommandLine(FString());
 		DefaultLaunchRole->SetInitialCulture(I18N.GetCurrentCulture()->GetName());
@@ -985,6 +1052,11 @@ public:
 			DeployedDeviceGroupId = FGuid();
 		}
 
+		if (DefaultDeployPlatform != NAME_None)
+		{
+			SetDefaultDeployPlatform(DefaultDeployPlatform);
+		}
+
 		Validate();
 	}
 
@@ -1446,6 +1518,9 @@ private:
 	// Holds the platforms to include in the build (only used if creating new builds).
 	TArray<FString> CookedPlatforms;
 
+	// Holds the platforms to deploy to if no specific devices were chosen for deploy.
+	FName DefaultDeployPlatform;
+
 	// Holds the default role (only used if launch mode is DefaultRole).
 	ILauncherProfileLaunchRoleRef DefaultLaunchRole;
 
diff --git a/Engine/Source/Developer/LauncherServices/Public/Interfaces/ILauncherDeviceGroup.h b/Engine/Source/Developer/LauncherServices/Public/Interfaces/ILauncherDeviceGroup.h
--- a/Engine/Source/Developer/LauncherServices/Public/Interfaces/ILauncherDeviceGroup.h
+++ b/Engine/Source/Developer/LauncherServices/Public/Interfaces/ILauncherDeviceGroup.h
@@ -76,6 +76,12 @@ public:
 	 */
 	virtual void RemoveDevice(const FString& DeviceID) = 0;
 
+	/**
+	* Removes all devices from the group.
+	*	
+	*/
+	virtual void RemoveAllDevices() = 0;
+
 	/**
 	 * Sets the name of the device group.
 	 *
diff --git a/Engine/Source/Developer/LauncherServices/Public/Interfaces/ILauncherProfile.h b/Engine/Source/Developer/LauncherServices/Public/Interfaces/ILauncherProfile.h
--- a/Engine/Source/Developer/LauncherServices/Public/Interfaces/ILauncherProfile.h
+++ b/Engine/Source/Developer/LauncherServices/Public/Interfaces/ILauncherProfile.h
@@ -525,7 +525,13 @@ public:
 	 * @return The device group, or NULL if none was configured.
 	 * @see SetDeployedDeviceGroup
 	 */
-	virtual ILauncherDeviceGroupPtr GetDeployedDeviceGroup( ) const = 0;
+	virtual ILauncherDeviceGroupPtr GetDeployedDeviceGroup( ) = 0;
+
+	/**
+	* Gets the default platforms to deploy if no specific devices were selected.
+	*	
+	*/
+	virtual const FName GetDefaultDeployPlatform() const = 0;
 
 	/**
 	 * Gets the deployment mode.
@@ -773,6 +779,15 @@ public:
 	 */
 	virtual void AddCookedPlatform( const FString& PlatformName ) = 0;
 
+	/**
+	* Adds a platform to deploy (only used if a specific device is not specified).
+	* Will deploy to the default device of the given platform, or the first device if none are
+	* marked as 'default'.
+	*
+	* @param PlatformName The name of the platform to add.		
+	*/
+	virtual void SetDefaultDeployPlatform(const FName PlatformName) = 0;
+
 	/**
 	 * Removes all cooked cultures.
 	 *
diff --git a/Engine/Source/Developer/ProjectLauncher/Private/Widgets/Deploy/SProjectLauncherDeployTargetListRow.h b/Engine/Source/Developer/ProjectLauncher/Private/Widgets/Deploy/SProjectLauncherDeployTargetListRow.h
--- a/Engine/Source/Developer/ProjectLauncher/Private/Widgets/Deploy/SProjectLauncherDeployTargetListRow.h
+++ b/Engine/Source/Developer/ProjectLauncher/Private/Widgets/Deploy/SProjectLauncherDeployTargetListRow.h
@@ -48,7 +48,6 @@ public:
 		HighlightText = InArgs._HighlightText;
 
 		SelectedVariant = NAME_None;
-		bDeviceInGroup = false;
 
 		ILauncherDeviceGroupPtr ActiveGroup = DeviceGroup.Get();
 		if (ActiveGroup.IsValid() && DeviceProxy.IsValid())
@@ -60,7 +59,6 @@ public:
 				if (DeviceProxy->HasDeviceId(DeviceID))
 				{
 					SelectedVariant = DeviceProxy->GetTargetDeviceVariant(DeviceID);
-					bDeviceInGroup = true;
 					break;
 				}
 			}
@@ -83,7 +81,7 @@ public:
 	{
 		if (ColumnName == "CheckBox")
 		{
-			return SNew(SCheckBox)
+			return SAssignNew(DeviceCheckbox, SCheckBox)
 				.IsChecked(this, &SProjectLauncherDeployTargetListRow::HandleCheckBoxIsChecked)
 				.OnCheckStateChanged(this, &SProjectLauncherDeployTargetListRow::HandleCheckBoxStateChanged)
 				.ToolTipText(LOCTEXT("CheckBoxToolTip", "Check this box to include this device in the current device group"));
@@ -184,12 +182,10 @@ private:
 			if (NewState == ECheckBoxState::Checked)
 			{
 				ActiveGroup->AddDevice(DeviceID);
-				bDeviceInGroup = true;
 			}
 			else
 			{
 				ActiveGroup->RemoveDevice(DeviceID);
-				bDeviceInGroup = false;
 			}
 		}
 	}
@@ -199,7 +195,13 @@ private:
 	{
 		if (IsEnabled())
 		{
-			return bDeviceInGroup ? ECheckBoxState::Checked : ECheckBoxState::Unchecked;
+			ILauncherDeviceGroupPtr ActiveGroup = DeviceGroup.Get();
+
+			if (ActiveGroup.IsValid() && DeviceProxy.IsValid() && DeviceProxy->HasVariant(SelectedVariant))
+			{
+				const FString& DeviceID = DeviceProxy->GetTargetDeviceId(SelectedVariant);
+				return ActiveGroup->GetDeviceIDs().Contains(DeviceID) ? ECheckBoxState::Checked : ECheckBoxState::Unchecked;
+			}
 		}
 
 		return ECheckBoxState::Unchecked;
@@ -226,11 +228,13 @@ private:
 	{
 		if (DeviceProxy.IsValid() && DeviceProxy->HasVariant(SelectedVariant) && DeviceProxy->HasVariant(InVariant))
 		{
-			if (bDeviceInGroup)
+			ILauncherDeviceGroupPtr ActiveGroup = DeviceGroup.Get();
+			const FString& DeviceID = DeviceProxy->GetTargetDeviceId(SelectedVariant);
+			if (ActiveGroup.IsValid() && ActiveGroup->GetDeviceIDs().Contains(DeviceID))
 			{
 				const FString& OldDeviceID = DeviceProxy->GetTargetDeviceId(SelectedVariant);
 				const FString& NewDeviceID = DeviceProxy->GetTargetDeviceId(InVariant);
-				ILauncherDeviceGroupPtr ActiveGroup = DeviceGroup.Get();
+				
 				if (ActiveGroup.IsValid())
 				{
 					ActiveGroup->RemoveDevice(OldDeviceID);
@@ -305,11 +309,10 @@ private:
 	// Holds a reference to the device proxy that is displayed in this row.
 	ITargetDeviceProxyPtr DeviceProxy;
 
-	// Holds the name of the selected variant.
-	FName SelectedVariant;
+	TSharedPtr<SCheckBox> DeviceCheckbox;
 
-	// Whether this device is executed on.
-	bool bDeviceInGroup;
+	// Holds the name of the selected variant.
+	FName SelectedVariant;	
 
 	// Holds the highlight string for the log message.
 	TAttribute<FText> HighlightText;
diff --git a/Engine/Source/Developer/ProjectLauncher/Private/Widgets/Deploy/SProjectLauncherDeployTargets.cpp b/Engine/Source/Developer/ProjectLauncher/Private/Widgets/Deploy/SProjectLauncherDeployTargets.cpp
--- a/Engine/Source/Developer/ProjectLauncher/Private/Widgets/Deploy/SProjectLauncherDeployTargets.cpp
+++ b/Engine/Source/Developer/ProjectLauncher/Private/Widgets/Deploy/SProjectLauncherDeployTargets.cpp
@@ -26,12 +26,28 @@ SProjectLauncherDeployTargets::~SProjectLauncherDeployTargets()
 BEGIN_SLATE_FUNCTION_BUILD_OPTIMIZATION
 void SProjectLauncherDeployTargets::Construct(const FArguments& InArgs, const FProjectLauncherModelRef& InModel)
 {
-	Model = InModel;
+	Model = InModel; 
 
 	ChildSlot
 	[
 		SNew(SVerticalBox)
 
+		+ SVerticalBox::Slot()
+		.AutoHeight()
+		.HAlign(HAlign_Left)
+		.Padding(10.0f, 10.0f)
+		[			
+			SAssignNew(PlatformComboBox, SComboBox<TSharedPtr<FName>>)			
+			.ContentPadding(FMargin(6.0f, 2.0f))
+			.OptionsSource(&VanillaPlatformList)
+			.OnGenerateWidget(this, &SProjectLauncherDeployTargets::HandlePlatformComboBoxGenerateWidget)
+			.OnSelectionChanged(this, &SProjectLauncherDeployTargets::HandlePlatformComboBoxSelectionChanged)
+			[
+				SNew(STextBlock)
+				.Text(this, &SProjectLauncherDeployTargets::HandlePlatformComboBoxContentText)
+			]
+		]
+
 		+ SVerticalBox::Slot()
 		.AutoHeight()
 		[
@@ -103,6 +119,25 @@ void SProjectLauncherDeployTargets::Construct(const FArguments& InArgs, const FP
 	DeviceProxyManager->OnProxyRemoved().AddSP(this, &SProjectLauncherDeployTargets::HandleDeviceProxyManagerProxyRemoved);
 
 	DeviceProxyManager->GetProxies(NAME_None, false, DeviceProxyList);
+
+	VanillaPlatformList.Reset();
+	TSet<FName> VanillaNames;
+	VanillaNames.Add(NAME_None);
+	VanillaPlatformList.Add(MakeShareable(new FName(NAME_None)));
+
+	TArray<ITargetPlatform*> Platforms = GetTargetPlatformManager()->GetTargetPlatforms();
+	if (Platforms.Num() > 0)
+	{
+		for (int32 PlatformIndex = 0; PlatformIndex < Platforms.Num(); ++PlatformIndex)
+		{
+			FName VanillaName = Platforms[PlatformIndex]->GetPlatformInfo().VanillaPlatformName;
+			if (!VanillaNames.Contains(VanillaName))
+			{				
+				VanillaPlatformList.AddUnique(MakeShareable(new FName(VanillaName)));
+			}
+		}
+	}
+	PlatformComboBox->SetSelectedItem(VanillaPlatformList[0]);
 }
 END_SLATE_FUNCTION_BUILD_OPTIMIZATION
 
@@ -181,7 +216,10 @@ TSharedRef<ITableRow> SProjectLauncherDeployTargets::HandleDeviceProxyListViewGe
 
 EVisibility SProjectLauncherDeployTargets::HandleDeviceProxyListViewVisibility() const
 {
-	if (DeviceProxyList.Num() > 0)
+	TSharedPtr<FName> DefaultPlatformNamePtr = PlatformComboBox->GetSelectedItem();
+	const FName DefaultPlatformName = DefaultPlatformNamePtr.IsValid() ? *DefaultPlatformNamePtr.Get() : NAME_None;	
+
+	if (DeviceProxyList.Num() > 0 && (DefaultPlatformName == NAME_None))
 	{
 		return EVisibility::Visible;
 	}
@@ -219,4 +257,55 @@ void SProjectLauncherDeployTargets::HandleDeviceProxyManagerProxyRemoved(const I
 	RefreshDeviceProxyList();
 }
 
+FText SProjectLauncherDeployTargets::HandlePlatformComboBoxContentText() const
+{	
+	FName DefaultPlatformName = NAME_None;
+	if (Model.IsValid())
+	{
+		const ILauncherProfilePtr Profile = Model->GetSelectedProfile();
+		if (Profile.IsValid())
+		{
+			DefaultPlatformName = Profile->GetDefaultDeployPlatform();
+		}
+	}
+
+	FString ComboText(TEXT("Default Deploy Platform: "));
+	ComboText += DefaultPlatformName.ToString();
+
+	//avoid creating new FTexts as much as possible.
+	if (DefaultPlatformText.ToString().Compare(ComboText) != 0)
+	{
+		DefaultPlatformText = FText::FromString(ComboText);
+		for (int32 i = 0; i < VanillaPlatformList.Num(); ++i)
+		{
+			if (*VanillaPlatformList[i] == DefaultPlatformName)
+			{
+				PlatformComboBox->SetSelectedItem(VanillaPlatformList[i]);
+				break;
+			}
+		}		
+	}
+
+	return DefaultPlatformText;
+}
+
+TSharedRef<SWidget> SProjectLauncherDeployTargets::HandlePlatformComboBoxGenerateWidget(TSharedPtr<FName> StringItem)
+{
+	FName DefaultPlatformName = StringItem.IsValid() ? *StringItem : NAME_None;	
+
+	return SNew(STextBlock).Text(FText::FromName(DefaultPlatformName));
+}
+
+void SProjectLauncherDeployTargets::HandlePlatformComboBoxSelectionChanged(TSharedPtr<FName> StringItem, ESelectInfo::Type SelectInfo)
+{
+	if (Model.IsValid())
+	{
+		const ILauncherProfilePtr Profile = Model->GetSelectedProfile();
+		if (Profile.IsValid())
+		{
+			Profile->SetDefaultDeployPlatform(StringItem.IsValid() ? *StringItem : NAME_None);
+		}
+	} 	
+}
+
 #undef LOCTEXT_NAMESPACE
diff --git a/Engine/Source/Developer/ProjectLauncher/Private/Widgets/Deploy/SProjectLauncherDeployTargets.h b/Engine/Source/Developer/ProjectLauncher/Private/Widgets/Deploy/SProjectLauncherDeployTargets.h
--- a/Engine/Source/Developer/ProjectLauncher/Private/Widgets/Deploy/SProjectLauncherDeployTargets.h
+++ b/Engine/Source/Developer/ProjectLauncher/Private/Widgets/Deploy/SProjectLauncherDeployTargets.h
@@ -75,11 +75,27 @@ private:
 	// Callback for when a device proxy has been removed from the device proxy manager.
 	void HandleDeviceProxyManagerProxyRemoved(const ITargetDeviceProxyRef& RemovedProxy);
 
+	// Callback for getting the name of the selected default platform.
+	FText HandlePlatformComboBoxContentText() const;
+
+	// Callback for generating widgets for the platforms combo box.
+	TSharedRef<SWidget> HandlePlatformComboBoxGenerateWidget(TSharedPtr<FName> StringItem);
+
+	// Callback for handling default platform selection changes.
+	void HandlePlatformComboBoxSelectionChanged(TSharedPtr<FName> StringItem, ESelectInfo::Type SelectInfo);
+
 private:
 
 	// Holds the list of available device proxies.
 	TArray<ITargetDeviceProxyPtr> DeviceProxyList;
 
+	// Holds the list of vanilla platforms;
+	TArray<TSharedPtr<FName>> VanillaPlatformList;
+
+	// Holds the list of vanilla platform options.
+	TSharedPtr<SComboBox<TSharedPtr<FName>>> PlatformComboBox;
+	mutable FText DefaultPlatformText;
+
 	// Holds the device proxy list view .
 	TSharedPtr<SListView<ITargetDeviceProxyPtr> > DeviceProxyListView;
 
