Properly save out Compress flag for UFE profiles
Fix UFE profile versioning.

diff --git a/Engine/Source/Developer/LauncherServices/Private/LauncherServicesPrivatePCH.h b/Engine/Source/Developer/LauncherServices/Private/LauncherServicesPrivatePCH.h
--- a/Engine/Source/Developer/LauncherServices/Private/LauncherServicesPrivatePCH.h
+++ b/Engine/Source/Developer/LauncherServices/Private/LauncherServicesPrivatePCH.h
@@ -20,14 +20,6 @@
 /* Private includes
  *****************************************************************************/
 
-/** Defines the launcher simple profile file version. */
-#define LAUNCHERSERVICES_SIMPLEPROFILEVERSION 1
-
-/** Defines the launcher profile file version. */
-#define LAUNCHERSERVICES_PROFILEVERSION 16 
-#define LAUNCHERSERVICES_MINPROFILEVERSION 10
-
-
 // profile manager
 #include "LauncherProjectPath.h"
 #include "LauncherDeviceGroup.h"
diff --git a/Engine/Source/Developer/LauncherServices/Private/Profiles/LauncherProfile.h b/Engine/Source/Developer/LauncherServices/Private/Profiles/LauncherProfile.h
--- a/Engine/Source/Developer/LauncherServices/Private/Profiles/LauncherProfile.h
+++ b/Engine/Source/Developer/LauncherServices/Private/Profiles/LauncherProfile.h
@@ -2,15 +2,31 @@
 
 #pragma once
 
-#define LAUNCHERSERVICES_ADDEDINCREMENTALDEPLOYVERSION 11
-#define LAUNCHERSERVICES_ADDEDPATCHSOURCECONTENTPATH 12
-#define LAUNCHERSERVICES_ADDEDRELEASEVERSION 13
-#define LAUNCHERSERVICES_REMOVEDPATCHSOURCECONTENTPATH 14
-#define LAUNCHERSERVICES_ADDEDDLCINCLUDEENGINECONTENT 15
-#define LAUNCHERSERVICES_ADDEDGENERATECHUNKS 16
-#define LAUNCHERSERVICES_ADDEDNUMCOOKERSTOSPAWN 17
-#define LAUNCHERSERVICES_ADDEDSKIPCOOKINGEDITORCONTENT 18
-#define LAUNCHERSERVICES_ADDEDDEFAULTDEPLOYPLATFORM 19
+enum ELauncherVersion
+{
+	LAUNCHERSERVICES_MINPROFILEVERSION=10,
+	LAUNCHERSERVICES_ADDEDINCREMENTALDEPLOYVERSION=11,
+	LAUNCHERSERVICES_ADDEDPATCHSOURCECONTENTPATH=12,
+	LAUNCHERSERVICES_ADDEDRELEASEVERSION=13,
+	LAUNCHERSERVICES_REMOVEDPATCHSOURCECONTENTPATH=14,
+	LAUNCHERSERVICES_ADDEDDLCINCLUDEENGINECONTENT=15,
+	LAUNCHERSERVICES_ADDEDGENERATECHUNKS=16,
+	LAUNCHERSERVICES_ADDEDNUMCOOKERSTOSPAWN=17,
+	LAUNCHERSERVICES_ADDEDSKIPCOOKINGEDITORCONTENT=18,
+	LAUNCHERSERVICES_ADDEDDEFAULTDEPLOYPLATFORM=19,
+	LAUNCHERSERVICES_FIXCOMPRESSIONSERIALIZE=20,
+
+	//ADD NEW STUFF HERE
+
+
+	LAUNCHERSERVICES_FINALPLUSONE,
+	LAUNCHERSERVICES_FINAL = LAUNCHERSERVICES_FINALPLUSONE-1,
+};
+
+enum ESimpleLauncherVersion
+{
+	LAUNCHERSERVICES_SIMPLEPROFILEVERSION=1,
+};
 
 /**
 * Implements a simple profile which controls the desired output of the Launcher for simple
@@ -679,7 +695,7 @@ public:
 
 	virtual bool Serialize( FArchive& Archive ) override
 	{
-		int32 Version = LAUNCHERSERVICES_ADDEDDEFAULTDEPLOYPLATFORM;
+		int32 Version = LAUNCHERSERVICES_FINAL;
 
 		Archive	<< Version;
 
@@ -688,6 +704,11 @@ public:
 			return false;
 		}
 
+		if (Version > LAUNCHERSERVICES_FINAL)
+		{
+			return false;
+		}
+
 		if (Archive.IsSaving())
 		{
 			if (DeployedDeviceGroup.IsValid())
@@ -728,6 +749,10 @@ public:
                 << Timeout;
 
 		FString DeployPlatformString = DefaultDeployPlatform.ToString();
+		if (Version >= LAUNCHERSERVICES_FIXCOMPRESSIONSERIALIZE)
+		{
+			Archive << Compressed;
+		}
 		if (Version >= LAUNCHERSERVICES_ADDEDDEFAULTDEPLOYPLATFORM)
 		{
 			Archive << DeployPlatformString;
