Make UFE store relative paths when possible for sharing.

diff --git a/Engine/Source/Developer/LauncherServices/Private/Profiles/LauncherProfile.h b/Engine/Source/Developer/LauncherServices/Private/Profiles/LauncherProfile.h
--- a/Engine/Source/Developer/LauncherServices/Private/Profiles/LauncherProfile.h
+++ b/Engine/Source/Developer/LauncherServices/Private/Profiles/LauncherProfile.h
@@ -15,6 +15,7 @@ enum ELauncherVersion
 	LAUNCHERSERVICES_ADDEDSKIPCOOKINGEDITORCONTENT=18,
 	LAUNCHERSERVICES_ADDEDDEFAULTDEPLOYPLATFORM=19,
 	LAUNCHERSERVICES_FIXCOMPRESSIONSERIALIZE=20,
+	LAUNCHERSERVICES_SHAREABLEPROJECTPATHS = 21,
 
 	//ADD NEW STUFF HERE
 
@@ -537,7 +538,7 @@ public:
 	{
 		if (ProjectSpecified)
 		{
-			FString Path = FLauncherProjectPath::GetProjectName(ProjectPath);
+			FString Path = FLauncherProjectPath::GetProjectName(FullProjectPath);
 			return Path;
 		}
 		return LauncherProfileManager->GetProjectName();
@@ -547,7 +548,7 @@ public:
 	{
 		if (ProjectSpecified)
 		{
-			return FLauncherProjectPath::GetProjectBasePath(ProjectPath);
+			return FLauncherProjectPath::GetProjectBasePath(FullProjectPath);
 		}
 		return LauncherProfileManager->GetProjectBasePath();
 	}
@@ -556,7 +557,7 @@ public:
 	{
 		if (ProjectSpecified)
 		{
-			return ProjectPath;
+			return FullProjectPath;
 		}
 		return LauncherProfileManager->GetProjectPath();
 	}
@@ -727,7 +728,7 @@ public:
 				<< Description
 				<< BuildConfiguration
 				<< ProjectSpecified
-				<< ProjectPath
+				<< ShareableProjectPath
 				<< CookConfiguration
 				<< CookIncremental
 				<< CookOptions
@@ -748,6 +749,11 @@ public:
                 << ForceClose
                 << Timeout;
 
+		if (LAUNCHERSERVICES_SHAREABLEPROJECTPATHS)
+		{
+			FullProjectPath = FPaths::ConvertRelativePathToFull(FPaths::RootDir(), ShareableProjectPath);
+		}
+
 		FString DeployPlatformString = DefaultDeployPlatform.ToString();
 		if (Version >= LAUNCHERSERVICES_FIXCOMPRESSIONSERIALIZE)
 		{
@@ -848,15 +854,15 @@ public:
 		// default project settings
 		if (FPaths::IsProjectFilePathSet())
 		{
-			ProjectPath = FPaths::GetProjectFilePath();
+			FullProjectPath = FPaths::GetProjectFilePath();
 		}
 		else if (FGameProjectHelper::IsGameAvailable(FApp::GetGameName()))
 		{
-			ProjectPath = FPaths::RootDir() / FApp::GetGameName() / FApp::GetGameName() + TEXT(".uproject");
+			FullProjectPath = FPaths::RootDir() / FApp::GetGameName() / FApp::GetGameName() + TEXT(".uproject");
 		}
 		else
 		{
-			ProjectPath = FString();
+			FullProjectPath = FString();
 		}
 
 		// Use the locally specified project path is resolving through the root isn't working
@@ -1228,15 +1234,28 @@ public:
 
 	virtual void SetProjectPath( const FString& Path ) override
 	{
-		if (ProjectPath != Path)
+		if (FullProjectPath != Path)
 		{
 			if(Path.IsEmpty())
 			{
-				ProjectPath = Path;
+				FullProjectPath = Path;
 			}
 			else
 			{
-				ProjectPath = FPaths::ConvertRelativePathToFull(Path);
+				FullProjectPath = FPaths::ConvertRelativePathToFull(Path);
+
+				FString RelativeProjectPath = Path;
+				bool bRelative = FPaths::MakePathRelativeTo(RelativeProjectPath, *FPaths::RootDir());
+
+				bool bIsUnderUE4Root = bRelative && !(RelativeProjectPath.StartsWith(FString("../"), ESearchCase::CaseSensitive));
+				if (bIsUnderUE4Root)
+				{
+					ShareableProjectPath = RelativeProjectPath;
+				}
+				else
+				{
+					ShareableProjectPath = FullProjectPath;
+				}				
 			}
 			CookedMaps.Reset();
 
@@ -1633,8 +1652,12 @@ private:
 	// Holds a flag indicating whether the project is specified by this profile.
 	bool ProjectSpecified;
 
-	// Holds the path to the Unreal project used by this profile.
-	FString ProjectPath;
+	// Holds the full absolute path to the Unreal project used by this profile.
+	FString FullProjectPath;
+
+	// Holds the path that might be shareable between people.  Only works if the project is under the UE4 root.
+	// otherwise this is an absolute path.
+	FString ShareableProjectPath;
 
 	// Holds the collection of validation errors.
 	TArray<ELauncherProfileValidationErrors::Type> ValidationErrors;
