properly utilize the same editor as running in for cooking
UE-20061

diff --git a/Engine/Source/Developer/LauncherServices/Private/Launcher/LauncherUATTask.h b/Engine/Source/Developer/LauncherServices/Private/Launcher/LauncherUATTask.h
--- a/Engine/Source/Developer/LauncherServices/Private/Launcher/LauncherUATTask.h
+++ b/Engine/Source/Developer/LauncherServices/Private/Launcher/LauncherUATTask.h
@@ -76,6 +76,12 @@ protected:
 		if(EditorExe.Len() > 0)
 		{
 			UATCommandLine += FString::Printf(TEXT(" -ue4exe=\"%s\""), *EditorExe);
+
+			bool bRunningDebug = FParse::Param(FCommandLine::Get(), TEXT("debug"));
+			if (bRunningDebug)
+			{
+				UATCommandLine += TEXT(" -UseDebugParamForEditorExe");
+			}
 		}
 
 		// specialized command arguments for this particular task
diff --git a/Engine/Source/Developer/LauncherServices/Private/Profiles/LauncherProfile.h b/Engine/Source/Developer/LauncherServices/Private/Profiles/LauncherProfile.h
--- a/Engine/Source/Developer/LauncherServices/Private/Profiles/LauncherProfile.h
+++ b/Engine/Source/Developer/LauncherServices/Private/Profiles/LauncherProfile.h
@@ -933,7 +933,29 @@ public:
 		PackagingMode = ELauncherProfilePackagingModes::DoNotPackage;
 
 		// default UAT settings
-		EditorExe.Empty();
+		EditorExe = FPlatformProcess::ExecutableName(false);
+		if (EditorExe.Contains(TEXT("Editor")))
+		{
+#if PLATFORM_WINDOWS
+			// turn UE4editor into UE4editor-cmd
+			if (EditorExe.EndsWith(".exe", ESearchCase::IgnoreCase) && !FPaths::GetBaseFilename(EditorExe).EndsWith("-cmd", ESearchCase::IgnoreCase))
+			{
+				FString NewExeName = EditorExe.Left(EditorExe.Len() - 4) + "-Cmd.exe";
+				if (FPaths::FileExists(NewExeName))
+				{
+					EditorExe = NewExeName;
+				}
+				else
+				{
+					EditorExe.Empty();
+				}
+			}
+#endif
+		}
+		else
+		{
+			EditorExe.Empty();
+		}
 
 		bNotForLicensees = false;
 
diff --git a/Engine/Source/Programs/AutomationTool/Scripts/CookCommand.Automation.cs b/Engine/Source/Programs/AutomationTool/Scripts/CookCommand.Automation.cs
--- a/Engine/Source/Programs/AutomationTool/Scripts/CookCommand.Automation.cs
+++ b/Engine/Source/Programs/AutomationTool/Scripts/CookCommand.Automation.cs
@@ -56,6 +56,10 @@ public partial class Project : CommandUtils
 				{
 					COTFCommandLine += " -iterate";
 				}
+				if (Params.UseDebugParamForEditorExe)
+				{
+					COTFCommandLine += " -debug";
+				}
 
 				var ServerLogFile = CombinePaths(LogFolderOutsideOfSandbox, "Server.log");
 				Platform ClientPlatformInst = Params.ClientTargetPlatformInstances[0];
