Always rebuild UAT by default when cooking/packaging, even if a debugger is present.

Many developers always run under the debugger, and this can cause confusing and seemingly unpredictable behavior by using a stale version of UAT rather than flagging that it needs to be built.

diff --git a/Engine/Source/Developer/LauncherServices/Private/Launcher/LauncherUATTask.h b/Engine/Source/Developer/LauncherServices/Private/Launcher/LauncherUATTask.h
--- a/Engine/Source/Developer/LauncherServices/Private/Launcher/LauncherUATTask.h
+++ b/Engine/Source/Developer/LauncherServices/Private/Launcher/LauncherUATTask.h
@@ -63,11 +63,7 @@ protected:
 
 		// we expect to pass -nocompile to UAT here as we generally expect UAT to be fully compiled. Besides, installed builds don't even have the source to compile UAT scripts.
 		// Only allow UAT to compile scripts dynamically if we pass -development or we have the IsBuildingUAT property set, the latter of which should not allow itself to be set in installed situations.
-#if PLATFORM_WINDOWS
-		bool bAllowCompile = !FPlatformMisc::IsDebuggerPresent();
-#else
 		bool bAllowCompile = true;
-#endif
 		UATCommandLine += (bAllowCompile && (FParse::Param( FCommandLine::Get(), TEXT("development") ) || ChainState.Profile->IsBuildingUAT())) ? TEXT("") : TEXT(" -nocompile");
 		// we never want to build the editor when launching from the editor or running with an installed engine (which can't rebuild itself)
 		UATCommandLine += GIsEditor || FApp::IsEngineInstalled() ? TEXT(" -nocompileeditor") : TEXT("");
diff --git a/Engine/Source/Editor/MainFrame/Private/Frame/MainFrameActions.cpp b/Engine/Source/Editor/MainFrame/Private/Frame/MainFrameActions.cpp
--- a/Engine/Source/Editor/MainFrame/Private/Frame/MainFrameActions.cpp
+++ b/Engine/Source/Editor/MainFrame/Private/Frame/MainFrameActions.cpp
@@ -333,7 +333,7 @@ const TCHAR* GetUATCompilationFlags()
 {
 	// We never want to compile editor targets when invoking UAT in this context.
 	// If we are installed or don't have a compiler, we must assume we have a precompiled UAT.
-	return (FApp::GetEngineIsPromotedBuild() || FApp::IsEngineInstalled() || FPlatformMisc::IsDebuggerPresent())
+	return (FApp::GetEngineIsPromotedBuild() || FApp::IsEngineInstalled())
 		? TEXT("-nocompile -nocompileeditor")
 		: TEXT("-nocompileeditor");
 }
