UE-58371 - Various actions utilizing Cook by the Book for certain platforms crash the editor

diff --git a/Engine/Source/Developer/Android/AndroidTargetPlatform/Private/AndroidTargetPlatform.h b/Engine/Source/Developer/Android/AndroidTargetPlatform/Private/AndroidTargetPlatform.h
--- a/Engine/Source/Developer/Android/AndroidTargetPlatform/Private/AndroidTargetPlatform.h
+++ b/Engine/Source/Developer/Android/AndroidTargetPlatform/Private/AndroidTargetPlatform.h
@@ -698,6 +698,57 @@ public:
 	}
 };
 
+/**
+* Android cooking platform which cooks only ETC1a based textures.
+*/
+class FAndroid_ETC1aTargetPlatform : public FAndroidTargetPlatform
+{
+public:
+
+	FAndroid_ETC1aTargetPlatform(bool bIsClient) : FAndroidTargetPlatform(bIsClient)
+	{
+		this->PlatformInfo = PlatformInfo::FindPlatformInfo("Android_ETC1a");
+	}
+
+	virtual FText DisplayName() const override
+	{
+		return LOCTEXT("Android_ETC1a", "Android (ETCa1)");
+	}
+
+	virtual FString GetAndroidVariantName() const override
+	{
+		return TEXT("ETC1a");
+	}
+
+	virtual bool SupportsTextureFormat(FName Format) const override
+	{
+		if (Format == AndroidTexFormat::NameAutoETC1a)
+		{
+			return true;
+		}
+
+		return false;
+	}
+
+	// End FAndroidTargetPlatform overrides
+
+	virtual bool SupportedByExtensionsString(const FString& ExtensionsString, const int GLESVersion) const override
+	{
+		return GLESVersion >= 0x30000;
+	}
+
+	virtual FText GetVariantDisplayName() const override
+	{
+		return LOCTEXT("Android_ETC1a_ShortName", "ETC1a");
+	}
+
+	virtual float GetVariantPriority() const override
+	{
+		float Priority;
+		return GConfig->GetFloat(TEXT("/Script/AndroidRuntimeSettings.AndroidRuntimeSettings"), TEXT("TextureFormatPriority_ETC1a"), Priority, GEngineIni) ? Priority : 1.0f;
+	}
+};
+
 
 
 
diff --git a/Engine/Source/Developer/Android/AndroidTargetPlatform/Private/AndroidTargetPlatformModule.cpp b/Engine/Source/Developer/Android/AndroidTargetPlatform/Private/AndroidTargetPlatformModule.cpp
--- a/Engine/Source/Developer/Android/AndroidTargetPlatform/Private/AndroidTargetPlatformModule.cpp
+++ b/Engine/Source/Developer/Android/AndroidTargetPlatform/Private/AndroidTargetPlatformModule.cpp
@@ -50,6 +50,7 @@ public:
 				SinglePlatforms.Add(new FAndroid_ATCTargetPlatform(bIsClient));
 				SinglePlatforms.Add(new FAndroid_DXTTargetPlatform(bIsClient));
 				SinglePlatforms.Add(new FAndroid_ETC1TargetPlatform(bIsClient));
+				SinglePlatforms.Add(new FAndroid_ETC1aTargetPlatform(bIsClient));
 				SinglePlatforms.Add(new FAndroid_ETC2TargetPlatform(bIsClient));
 				SinglePlatforms.Add(new FAndroid_PVRTCTargetPlatform(bIsClient));
 
diff --git a/Engine/Source/Developer/Android/Android_ETC1aTargetPlatform/Private/Android_ETC1aTargetPlatformModule.cpp b/Engine/Source/Developer/Android/Android_ETC1aTargetPlatform/Private/Android_ETC1aTargetPlatformModule.cpp
--- a/Engine/Source/Developer/Android/Android_ETC1aTargetPlatform/Private/Android_ETC1aTargetPlatformModule.cpp
+++ b/Engine/Source/Developer/Android/Android_ETC1aTargetPlatform/Private/Android_ETC1aTargetPlatformModule.cpp
@@ -17,61 +17,6 @@
 #define LOCTEXT_NAMESPACE "FAndroid_ETC1aTargetPlatformModule" 
 
 
-/**
- * Android cooking platform which cooks only ETC1a based textures.
- */
-class FAndroid_ETC1aTargetPlatform
-	: public FAndroidTargetPlatform<FAndroid_ETC1aPlatformProperties>
-{
-public:
-
-	// Begin FAndroidTargetPlatform overrides
-
-	virtual FText DisplayName() const override
-	{
-		return LOCTEXT("Android_ETC1a", "Android (ETC1a)");
-	}
-
-	virtual FString GetAndroidVariantName()
-	{
-		return TEXT("ETC1a");
-	}
-
-	virtual FString PlatformName() const override
-	{
-		return FString(FAndroid_ETC1aPlatformProperties::PlatformName());
-	}
-
-	virtual bool SupportsTextureFormat(FName Format) const override
-	{
-		if (Format == AndroidTexFormat::NameAutoETC1a)
-		{
-			return true;
-		}
-
-		return false;
-	}
-
-	// End FAndroidTargetPlatform overrides
-
-	virtual bool SupportedByExtensionsString(const FString& ExtensionsString, const int GLESVersion) const override
-	{
-		return GLESVersion >= 0x30000;
-	}
-
-	virtual FText GetVariantDisplayName() const override
-	{
-		return LOCTEXT("Android_ETC1a_ShortName", "ETC1a");
-	}
-
-	virtual float GetVariantPriority() const override
-	{
-		float Priority;
-		return GConfig->GetFloat(TEXT("/Script/AndroidRuntimeSettings.AndroidRuntimeSettings"), TEXT("TextureFormatPriority_ETC1a"), Priority, GEngineIni) ? Priority : 1.0f;
-	}
-};
-
-
 /**
  * Holds the target platform singleton.
  */
diff --git a/Engine/Source/Developer/LauncherServices/Private/Launcher/LauncherWorker.cpp b/Engine/Source/Developer/LauncherServices/Private/Launcher/LauncherWorker.cpp
--- a/Engine/Source/Developer/LauncherServices/Private/Launcher/LauncherWorker.cpp
+++ b/Engine/Source/Developer/LauncherServices/Private/Launcher/LauncherWorker.cpp
@@ -309,62 +309,63 @@ FString FLauncherWorker::CreateUATCommand( const ILauncherProfileRef& InProfile,
 		// Platform info for the given platform
 		const PlatformInfo::FPlatformInfo* PlatformInfo = PlatformInfo::FindPlatformInfo(FName(*InPlatforms[PlatformIndex]));
 
-		// switch server and no editor platforms to the proper type
-		if (PlatformInfo->TargetPlatformName == FName("LinuxServer"))
+		if (ensure(PlatformInfo))
 		{
-			ServerPlatforms += TEXT("+Linux");
-		}
-		else if (PlatformInfo->TargetPlatformName == FName("WindowsServer"))
-		{
-			ServerPlatforms += TEXT("+Win64");
-		}
-		else if (PlatformInfo->TargetPlatformName == FName("MacServer"))
-		{
-			ServerPlatforms += TEXT("+Mac");
-		}
-		else if (PlatformInfo->TargetPlatformName == FName("LinuxNoEditor") || PlatformInfo->TargetPlatformName == FName("LinuxClient"))
-		{
-			Platforms += TEXT("+Linux");
-		}
-		else if (PlatformInfo->TargetPlatformName == FName("WindowsNoEditor") || PlatformInfo->TargetPlatformName == FName("Windows") || PlatformInfo->TargetPlatformName == FName("WindowsClient"))
-		{
-			Platforms += TEXT("+Win64");
-		}
-		else if (PlatformInfo->TargetPlatformName == FName("MacNoEditor") || PlatformInfo->TargetPlatformName == FName("MacClient"))
-		{
-			Platforms += TEXT("+Mac");
-		}
-		else
-		{
-			Platforms += TEXT("+");
-			Platforms += PlatformInfo->TargetPlatformName.ToString();
-		}
-
-		// Append any extra UAT flags specified for this platform flavor
-		if (!PlatformInfo->UATCommandLine.IsEmpty())
-		{
-			FString OptionalUATCommandLine = PlatformInfo->UATCommandLine;
-			
-			FString OptionalTargetPlatform;
-			if (FParse::Value(*OptionalUATCommandLine, TEXT("-targetplatform="), OptionalTargetPlatform))
+			// switch server and no editor platforms to the proper type
+			if (PlatformInfo->TargetPlatformName == FName("LinuxServer"))
 			{
-				OptionalTargetPlatforms.Add(OptionalTargetPlatform);
-				OptionalUATCommandLine.ReplaceInline(*(TEXT("-targetplatform=") + OptionalTargetPlatform), TEXT(""));
+				ServerPlatforms += TEXT("+Linux");
 			}
-
-			FString OptionalCookFlavor;
-			if (FParse::Value(*OptionalUATCommandLine, TEXT("-cookflavor="), OptionalCookFlavor))
+			else if (PlatformInfo->TargetPlatformName == FName("WindowsServer"))
 			{
-				OptionalCookFlavors.Add(OptionalCookFlavor);
-				OptionalUATCommandLine.ReplaceInline(*(TEXT("-cookflavor=") + OptionalCookFlavor), TEXT(""));
+				ServerPlatforms += TEXT("+Win64");
+			}
+			else if (PlatformInfo->TargetPlatformName == FName("MacServer"))
+			{
+				ServerPlatforms += TEXT("+Mac");
+			}
+			else if (PlatformInfo->TargetPlatformName == FName("LinuxNoEditor") || PlatformInfo->TargetPlatformName == FName("LinuxClient"))
+			{
+				Platforms += TEXT("+Linux");
+			}
+			else if (PlatformInfo->TargetPlatformName == FName("WindowsNoEditor") || PlatformInfo->TargetPlatformName == FName("Windows") || PlatformInfo->TargetPlatformName == FName("WindowsClient"))
+			{
+				Platforms += TEXT("+Win64");
+			}
+			else if (PlatformInfo->TargetPlatformName == FName("MacNoEditor") || PlatformInfo->TargetPlatformName == FName("MacClient"))
+			{
+				Platforms += TEXT("+Mac");
+			}
+			else
+			{
+				Platforms += TEXT("+");
+				Platforms += PlatformInfo->TargetPlatformName.ToString();
 			}
-			
-			OptionalParams += TEXT(" ");
-			OptionalParams += OptionalUATCommandLine;
-		}
 
-		bUATClosesAfterLaunch |= PlatformInfo->bUATClosesAfterLaunch;
+			// Append any extra UAT flags specified for this platform flavor
+			if (!PlatformInfo->UATCommandLine.IsEmpty())
+			{
+				FString OptionalUATCommandLine = PlatformInfo->UATCommandLine;
 
+				FString OptionalTargetPlatform;
+				if (FParse::Value(*OptionalUATCommandLine, TEXT("-targetplatform="), OptionalTargetPlatform))
+				{
+					OptionalTargetPlatforms.Add(OptionalTargetPlatform);
+					OptionalUATCommandLine.ReplaceInline(*(TEXT("-targetplatform=") + OptionalTargetPlatform), TEXT(""));
+				}
+
+				FString OptionalCookFlavor;
+				if (FParse::Value(*OptionalUATCommandLine, TEXT("-cookflavor="), OptionalCookFlavor))
+				{
+					OptionalCookFlavors.Add(OptionalCookFlavor);
+					OptionalUATCommandLine.ReplaceInline(*(TEXT("-cookflavor=") + OptionalCookFlavor), TEXT(""));
+				}
+
+				OptionalParams += TEXT(" ");
+				OptionalParams += OptionalUATCommandLine;
+			}
+			bUATClosesAfterLaunch |= PlatformInfo->bUATClosesAfterLaunch;
+		}
 	}
 	if (ServerPlatforms.Len() > 0)
 	{
