Make UFE/ProjectLauncher use the vanilla platform name for -platform when using Android cook flavors, rather than eg -platform=Android_ETC1

The old SetupTargetPlatforms() had a comment "// Need to tolerate cook platform names here, which UFE likes to pass in." but the data driven UnrealTargetPlatform Parse() no longer accepts non-canonical values for -platform. UFE always passes -targetplatform=Android and -cookflavor=ETC1 anyway. In the old code, the invalid -platform was ignored and the -targetplatform=Android was used. With this change the new code parses both and removes duplicates.

diff --git a/Engine/Source/Developer/LauncherServices/Private/Launcher/LauncherWorker.cpp b/Engine/Source/Developer/LauncherServices/Private/Launcher/LauncherWorker.cpp
--- a/Engine/Source/Developer/LauncherServices/Private/Launcher/LauncherWorker.cpp
+++ b/Engine/Source/Developer/LauncherServices/Private/Launcher/LauncherWorker.cpp
@@ -347,7 +347,7 @@ FString FLauncherWorker::CreateUATCommand( const ILauncherProfileRef& InProfile,
 			else
 			{
 				Platforms += TEXT("+");
-				Platforms += PlatformInfo->TargetPlatformName.ToString();
+				Platforms += PlatformInfo->VanillaPlatformName.ToString();
 			}
 
 			// Append any extra UAT flags specified for this platform flavor
