Fix issue where deploying to target platforms would always trigger a compile.

#rnx

diff --git a/Engine/Source/Developer/LauncherServices/Private/Profiles/LauncherProfile.h b/Engine/Source/Developer/LauncherServices/Private/Profiles/LauncherProfile.h
--- a/Engine/Source/Developer/LauncherServices/Private/Profiles/LauncherProfile.h
+++ b/Engine/Source/Developer/LauncherServices/Private/Profiles/LauncherProfile.h
@@ -1781,7 +1781,7 @@ public:
 		}
 
 		// Use the locally specified project path is resolving through the root isn't working
-		ProjectSpecified = GetProjectPath().IsEmpty();
+		ProjectSpecified = !GetProjectPath().IsEmpty();
 		
 		// I don't use FApp::GetBuildConfiguration() because i don't want the act of running in debug the first time to cause 
 		// profiles the user creates to be in debug. This will keep consistency.
