PR #4937: UE-62508: Add option to project launcher to include prerequisites com. (Contributed by [REDACTED_PERSON])

diff --git a/Engine/Source/Developer/LauncherServices/Private/Launcher/LauncherWorker.cpp b/Engine/Source/Developer/LauncherServices/Private/Launcher/LauncherWorker.cpp
--- a/Engine/Source/Developer/LauncherServices/Private/Launcher/LauncherWorker.cpp
+++ b/Engine/Source/Developer/LauncherServices/Private/Launcher/LauncherWorker.cpp
@@ -777,6 +777,11 @@ FString FLauncherWorker::CreateUATCommand( const ILauncherProfileRef& InProfile,
 	}
 	else
 	{
+		if (InProfile->IsIncludingPrerequisites())
+		{
+			UATCommand += TEXT(" -prereqs");
+		}
+
 		if (InProfile->GetPackagingMode() == ELauncherProfilePackagingModes::Locally)
 		{
 			UATCommand += TEXT(" -stage -package");
diff --git a/Engine/Source/Developer/LauncherServices/Private/Profiles/LauncherProfile.h b/Engine/Source/Developer/LauncherServices/Private/Profiles/LauncherProfile.h
--- a/Engine/Source/Developer/LauncherServices/Private/Profiles/LauncherProfile.h
+++ b/Engine/Source/Developer/LauncherServices/Private/Profiles/LauncherProfile.h
@@ -41,6 +41,7 @@ enum ELauncherVersion
 	LAUNCHERSERVICES_ADDEDENCRYPTINIFILES = 24,
 	LAUNCHERSERVICES_ADDEDMULTILEVELPATCHING = 25,
 	LAUNCHERSERVICES_ADDEDADDITIONALCOMMANDLINE = 26,
+	LAUNCHERSERVICES_ADDEDINCLUDEPREREQUISITES = 27,
 
 	//ADD NEW STUFF HERE
 
@@ -245,8 +246,6 @@ public:
 		}
 	}
 
-public:
-
 	/**
 	 * Gets the identifier of the device group to deploy to.
 	 *
@@ -261,8 +260,6 @@ public:
 		return DeployedDeviceGroupId;
 	}
 
-public:
-
 	//~ Begin ILauncherProfile Interface
 
 	virtual void AddCookedCulture( const FString& CultureName ) override
@@ -739,6 +736,11 @@ public:
 		return DeployWithUnrealPak;
 	}
 
+	virtual bool IsIncludingPrerequisites() const override
+	{
+		return IncludePrerequisites;
+	}
+
 	virtual bool IsGeneratingChunks() const override
 	{
 		return bGenerateChunks;
@@ -919,6 +921,14 @@ public:
 			Archive << bArchive;
 			Archive << ArchiveDir;
 		}
+		if (Version >= LAUNCHERSERVICES_ADDEDADDITIONALCOMMANDLINE)
+		{
+			Archive << AdditionalCommandLineParameters;
+		}
+		if (Version >= LAUNCHERSERVICES_ADDEDINCLUDEPREREQUISITES)
+		{
+			Archive << IncludePrerequisites;
+		}
 		
 		DefaultLaunchRole->Serialize(Archive);
 
@@ -1054,6 +1064,7 @@ public:
 		Writer.WriteValue("Archive", bArchive);
 		Writer.WriteValue("ArchiveDirectory", ArchiveDir);
 		Writer.WriteValue("AdditionalCommandLineParameters", AdditionalCommandLineParameters);
+		Writer.WriteValue("IncludePrerequisites", IncludePrerequisites);
 
 		// serialize the default launch role
 		DefaultLaunchRole->Save(Writer, TEXT("DefaultRole"));
@@ -1714,6 +1725,11 @@ public:
 			AdditionalCommandLineParameters = TEXT("");
 		}
 
+		if (Version >= LAUNCHERSERVICES_ADDEDINCLUDEPREREQUISITES)
+		{
+			IncludePrerequisites = Object.GetBoolField("IncludePrerequisites");
+		}
+
 		// load the default launch role
 		TSharedPtr<FJsonObject> Role = Object.GetObjectField("DefaultRole");
 		DefaultLaunchRole->Load(*(Role.Get()));
@@ -1807,6 +1823,7 @@ public:
 		DeployedDeviceGroup.Reset();
 		DeploymentMode = ELauncherProfileDeploymentModes::CopyToDevice;
 		DeployStreamingServer = false;
+		IncludePrerequisites = false;
 		DeployWithUnrealPak = false;
 		DeployedDeviceGroupId = FGuid();
 		HideFileServerWindow = false;
@@ -2263,6 +2280,16 @@ public:
 		}
 	}
 
+	virtual void SetIncludePrerequisites(bool InValue) override
+	{
+		if (IncludePrerequisites != InValue)
+		{
+			IncludePrerequisites = InValue;
+
+			Validate();
+		}
+	}
+
     virtual void SetTimeout( uint32 InTime ) override
     {
         if (Timeout != InTime)
@@ -2647,6 +2674,9 @@ private:
 	// Holds a flag indicating whether content should be packaged with UnrealPak.
 	bool DeployWithUnrealPak;
 
+	// Flag to indicate if game prerequisites should be included
+	bool IncludePrerequisites;
+
 	// Flag indicating if content should be split into chunks
 	bool bGenerateChunks;
 	
diff --git a/Engine/Source/Developer/LauncherServices/Public/ILauncherProfile.h b/Engine/Source/Developer/LauncherServices/Public/ILauncherProfile.h
--- a/Engine/Source/Developer/LauncherServices/Public/ILauncherProfile.h
+++ b/Engine/Source/Developer/LauncherServices/Public/ILauncherProfile.h
@@ -803,6 +803,14 @@ public:
 	 */
 	virtual bool IsPackingWithUnrealPak( ) const = 0;
 
+	/**
+	 * Checks whether to include an installer for prerequisites of packaged games, such as redistributable operating system components, on platforms that support it.
+	 *
+	 * @return true if prerequisites are to be included, false otherwise.
+	 * @see SetIncludePrerequisites
+	 */
+	virtual bool IsIncludingPrerequisites() const = 0;
+
 	/**
 	 * Return whether packaging will generate chunk data.
 	 *
@@ -1218,6 +1226,14 @@ public:
 	 */
 	virtual void SetStreamingFileServer( bool Streaming ) = 0;
 
+	/**
+	 * Sets whether to include game prerequisites.
+	 *
+	 * @param Value Whether prerequisites should be used.
+	 * @see IsIncludingPrerequisites
+	 */
+	virtual void SetIncludePrerequisites(bool InValue) = 0;
+
     /**
      * Sets the cook on the fly server timeout
      *
diff --git a/Engine/Source/Developer/ProjectLauncher/Private/Widgets/Package/SProjectLauncherPackagingSettings.cpp b/Engine/Source/Developer/ProjectLauncher/Private/Widgets/Package/SProjectLauncherPackagingSettings.cpp
--- a/Engine/Source/Developer/ProjectLauncher/Private/Widgets/Package/SProjectLauncherPackagingSettings.cpp
+++ b/Engine/Source/Developer/ProjectLauncher/Private/Widgets/Package/SProjectLauncherPackagingSettings.cpp
@@ -109,6 +109,24 @@ void SProjectLauncherPackagingSettings::Construct(const FArguments& InArgs, cons
 										.Text(LOCTEXT("ForDistributionCheckBoxText", "Is this build for distribution to the public"))
 									]
 							]
+
+
+						+ SVerticalBox::Slot()
+							.AutoHeight()
+							.Padding(0.0, 4.0, 0.0, 0.0)
+							[
+								SNew(SCheckBox)
+								.IsEnabled(this, &SProjectLauncherPackagingSettings::IsEditable)
+								.IsChecked(this, &SProjectLauncherPackagingSettings::HandleIncludePrerequisitesCheckBoxIsChecked)
+								.OnCheckStateChanged(this, &SProjectLauncherPackagingSettings::HandleIncludePrerequisitesCheckStateChanged)
+								.Padding(FMargin(4.0f, 0.0f))
+								.ToolTipText(LOCTEXT("IncludePrerequisitesCheckBoxTooltip", "If checked the build will include the prerequisites installer on platforms that support it."))
+								.Content()
+								[
+									SNew(STextBlock)
+									.Text(LOCTEXT("IncludePrerequisitesCheckBoxText", "Include an installer for prerequisites of packaged games"))
+								]
+							]
 					]
 			]
 	];
@@ -150,6 +168,27 @@ ECheckBoxState SProjectLauncherPackagingSettings::HandleForDistributionCheckBoxI
 	return ECheckBoxState::Unchecked;
 }
 
+void SProjectLauncherPackagingSettings::HandleIncludePrerequisitesCheckStateChanged(ECheckBoxState NewState)
+{
+	ILauncherProfilePtr SelectedProfile = Model->GetSelectedProfile();
+	if (SelectedProfile.IsValid())
+	{
+		SelectedProfile->SetIncludePrerequisites(NewState == ECheckBoxState::Checked);
+	}
+}
+
+ECheckBoxState SProjectLauncherPackagingSettings::HandleIncludePrerequisitesCheckBoxIsChecked() const
+{
+	ILauncherProfilePtr SelectedProfile = Model->GetSelectedProfile();
+	if (SelectedProfile.IsValid())
+	{
+		if (SelectedProfile->IsIncludingPrerequisites())
+		{
+			return ECheckBoxState::Checked;
+		}
+	}
+	return ECheckBoxState::Unchecked;
+}
 
 FText SProjectLauncherPackagingSettings::HandleDirectoryTitleText() const
 {
diff --git a/Engine/Source/Developer/ProjectLauncher/Private/Widgets/Package/SProjectLauncherPackagingSettings.h b/Engine/Source/Developer/ProjectLauncher/Private/Widgets/Package/SProjectLauncherPackagingSettings.h
--- a/Engine/Source/Developer/ProjectLauncher/Private/Widgets/Package/SProjectLauncherPackagingSettings.h
+++ b/Engine/Source/Developer/ProjectLauncher/Private/Widgets/Package/SProjectLauncherPackagingSettings.h
@@ -44,10 +44,12 @@ public:
 
 private:
 
-
 	void HandleForDistributionCheckBoxCheckStateChanged(ECheckBoxState NewState);
 	ECheckBoxState HandleForDistributionCheckBoxIsChecked() const;
 
+	void HandleIncludePrerequisitesCheckStateChanged(ECheckBoxState NewState);
+	ECheckBoxState HandleIncludePrerequisitesCheckBoxIsChecked() const;
+
 	// Callback for getting the content text of the 'Directory' label.
 	FText HandleDirectoryTitleText() const;
 	FText HandleDirectoryPathText() const;
