Fixed add patch tier in project launcher passing the wrong commandline option to UAT.
#test none

diff --git a/Engine/Source/Developer/LauncherServices/Private/Launcher/LauncherWorker.cpp b/Engine/Source/Developer/LauncherServices/Private/Launcher/LauncherWorker.cpp
--- a/Engine/Source/Developer/LauncherServices/Private/Launcher/LauncherWorker.cpp
+++ b/Engine/Source/Developer/LauncherServices/Private/Launcher/LauncherWorker.cpp
@@ -574,7 +574,7 @@ FString FLauncherWorker::CreateUATCommand( const ILauncherProfileRef& InProfile,
 
 				if ( InProfile->ShouldAddPatchLevel() )
 				{
-					UATCommand += TEXT(" -newpatchlevel");
+					UATCommand += TEXT(" -addpatchlevel");
 				}
 			}
 
