PR #5426: Fix incorrect check for LAUNCHERSERVICES_SHAREABLEPROJECTPATHS (Contributed by [REDACTED_PERSON])

diff --git a/Engine/Source/Developer/LauncherServices/Private/Profiles/LauncherProfile.h b/Engine/Source/Developer/LauncherServices/Private/Profiles/LauncherProfile.h
--- a/Engine/Source/Developer/LauncherServices/Private/Profiles/LauncherProfile.h
+++ b/Engine/Source/Developer/LauncherServices/Private/Profiles/LauncherProfile.h
@@ -1746,7 +1746,7 @@ public:
 			}
 		}
 
-		if (LAUNCHERSERVICES_SHAREABLEPROJECTPATHS)
+		if (Version >= LAUNCHERSERVICES_SHAREABLEPROJECTPATHS)
 		{
 			FullProjectPath = FPaths::ConvertRelativePathToFull(FPaths::RootDir(), ShareableProjectPath);
 		}
