Removed Redundant if statement

diff --git a/Engine/Source/Developer/LauncherServices/Private/Profiles/LauncherProfile.h b/Engine/Source/Developer/LauncherServices/Private/Profiles/LauncherProfile.h
--- a/Engine/Source/Developer/LauncherServices/Private/Profiles/LauncherProfile.h
+++ b/Engine/Source/Developer/LauncherServices/Private/Profiles/LauncherProfile.h
@@ -1746,10 +1746,11 @@ public:
 			}
 		}
 
-		if (Version >= LAUNCHERSERVICES_SHAREABLEPROJECTPATHS)
+		// if (Version >= LAUNCHERSERVICES_SHAREABLEPROJECTPATHS) Always true due to early out at top of function
 		{
 			FullProjectPath = FPaths::ConvertRelativePathToFull(FPaths::RootDir(), ShareableProjectPath);
 		}
+
 		if (DefaultDeployPlatform != NAME_None)
 		{
 			SetDefaultDeployPlatform(DefaultDeployPlatform);
