Rename EBuildConfigurations::Type to EBuildConfiguration.

diff --git a/Engine/Platforms/XXX/Source/Developer/XXXTargetPlatform/Private/XXXTargetDevice.cpp b/Engine/Platforms/XXX/Source/Developer/XXXTargetPlatform/Private/XXXTargetDevice.cpp
--- a/Engine/Platforms/XXX/Source/Developer/XXXTargetPlatform/Private/XXXTargetDevice.cpp
+++ b/Engine/Platforms/XXX/Source/Developer/XXXTargetPlatform/Private/XXXTargetDevice.cpp
@@ -46,7 +46,7 @@ int32 FXXXTargetDevice::GetProcessSnapshot( TArray<FTargetDeviceProcessInfo>& Ou
 }
 
 
-bool FXXXTargetDevice::Launch( const FString& AppId, EBuildConfigurations::Type BuildConfiguration, EBuildTargetType TargetType, const FString& Params, uint32* OutProcessId )
+bool FXXXTargetDevice::Launch( const FString& AppId, EBuildConfiguration BuildConfiguration, EBuildTargetType TargetType, const FString& Params, uint32* OutProcessId )
 {
 	checkNoEntry();
 	return false;
diff --git a/Engine/Platforms/XXX/Source/Developer/XXXTargetPlatform/Private/XXXTargetDevice.h b/Engine/Platforms/XXX/Source/Developer/XXXTargetPlatform/Private/XXXTargetDevice.h
--- a/Engine/Platforms/XXX/Source/Developer/XXXTargetPlatform/Private/XXXTargetDevice.h
+++ b/Engine/Platforms/XXX/Source/Developer/XXXTargetPlatform/Private/XXXTargetDevice.h
@@ -97,7 +97,7 @@ public:
 		return CachedDefault;
 	}
 
-	virtual bool Launch( const FString& AppId, EBuildConfigurations::Type BuildConfiguration, EBuildTargetType TargetType, const FString& Params, uint32* OutProcessId ) override;
+	virtual bool Launch( const FString& AppId, EBuildConfiguration BuildConfiguration, EBuildTargetType TargetType, const FString& Params, uint32* OutProcessId ) override;
 
 	virtual bool PowerOff( bool Force ) override;
 
diff --git a/Engine/Plugins/Developer/Concert/DisasterRecoveryClient/Source/DisasterRecoveryClient/Private/DisasterRecoveryClientModule.cpp b/Engine/Plugins/Developer/Concert/DisasterRecoveryClient/Source/DisasterRecoveryClient/Private/DisasterRecoveryClientModule.cpp
--- a/Engine/Plugins/Developer/Concert/DisasterRecoveryClient/Source/DisasterRecoveryClient/Private/DisasterRecoveryClientModule.cpp
+++ b/Engine/Plugins/Developer/Concert/DisasterRecoveryClient/Source/DisasterRecoveryClient/Private/DisasterRecoveryClientModule.cpp
@@ -228,7 +228,7 @@ private:
 
 	static FString GetDisasterRecoveryServicePath()
 	{
-		auto GetDisasterRecoveryServicePathForBuildConfiguration = [](const EBuildConfigurations::Type InBuildConfiguration) -> FString
+		auto GetDisasterRecoveryServicePathForBuildConfiguration = [](const EBuildConfiguration InBuildConfiguration) -> FString
 		{
 			FString ServicePath = FPlatformProcess::GenerateApplicationPath(TEXT("UnrealDisasterRecoveryService"), InBuildConfiguration);
 			return FPaths::FileExists(ServicePath) ? ServicePath : FString();
@@ -238,9 +238,9 @@ private:
 		FString ServicePath = GetDisasterRecoveryServicePathForBuildConfiguration(FApp::GetBuildConfiguration());
 
 		// Fall back to Development if the app doesn't exist for our build configuration, as installed builds only build it for Development
-		if (ServicePath.IsEmpty() && FApp::GetBuildConfiguration() != EBuildConfigurations::Development)
+		if (ServicePath.IsEmpty() && FApp::GetBuildConfiguration() != EBuildConfiguration::Development)
 		{
-			ServicePath = GetDisasterRecoveryServicePathForBuildConfiguration(EBuildConfigurations::Development);
+			ServicePath = GetDisasterRecoveryServicePathForBuildConfiguration(EBuildConfiguration::Development);
 		}
 
 		return ServicePath;
diff --git a/Engine/Plugins/Developer/Concert/MultiUserClient/Source/MultiUserClient/Private/MultiUserClientModule.cpp b/Engine/Plugins/Developer/Concert/MultiUserClient/Source/MultiUserClient/Private/MultiUserClientModule.cpp
--- a/Engine/Plugins/Developer/Concert/MultiUserClient/Source/MultiUserClient/Private/MultiUserClientModule.cpp
+++ b/Engine/Plugins/Developer/Concert/MultiUserClient/Source/MultiUserClient/Private/MultiUserClientModule.cpp
@@ -844,9 +844,9 @@ private:
 		FString ServerPath = FPlatformProcess::GenerateApplicationPath(MultiUserServerName, FApp::GetBuildConfiguration());
 
 		// Validate it exists and fall back to development if it doesn't.
-		if (!IFileManager::Get().FileExists(*ServerPath) && FApp::GetBuildConfiguration() != EBuildConfigurations::Development)
+		if (!IFileManager::Get().FileExists(*ServerPath) && FApp::GetBuildConfiguration() != EBuildConfiguration::Development)
 		{
-			ServerPath = FPlatformProcess::GenerateApplicationPath(MultiUserServerName, EBuildConfigurations::Development);
+			ServerPath = FPlatformProcess::GenerateApplicationPath(MultiUserServerName, EBuildConfiguration::Development);
 		}
 
 		return FPlatformProcess::IsApplicationRunning(*FPaths::GetCleanFilename(ServerPath));
@@ -876,9 +876,9 @@ private:
 		FString ServerPath = FPlatformProcess::GenerateApplicationPath(MultiUserServerName, FApp::GetBuildConfiguration());
 
 		// Validate it exists and fall back to development if it doesn't
-		if (!IFileManager::Get().FileExists(*ServerPath) && FApp::GetBuildConfiguration() != EBuildConfigurations::Development)
+		if (!IFileManager::Get().FileExists(*ServerPath) && FApp::GetBuildConfiguration() != EBuildConfiguration::Development)
 		{
-			ServerPath = FPlatformProcess::GenerateApplicationPath(MultiUserServerName, EBuildConfigurations::Development);
+			ServerPath = FPlatformProcess::GenerateApplicationPath(MultiUserServerName, EBuildConfiguration::Development);
 		}
 		
 		FText LaunchMultiUserErrorTitle = LOCTEXT("LaunchUnrealMultiUserServerErrorTitle", "Failed to Launch the Unreal Multi-User Server");
diff --git a/Engine/Plugins/Editor/MobileLauncherProfileWizard/Source/MobileLauncherProfileWizard/Private/SProfileWizardUI.cpp b/Engine/Plugins/Editor/MobileLauncherProfileWizard/Source/MobileLauncherProfileWizard/Private/SProfileWizardUI.cpp
--- a/Engine/Plugins/Editor/MobileLauncherProfileWizard/Source/MobileLauncherProfileWizard/Private/SProfileWizardUI.cpp
+++ b/Engine/Plugins/Editor/MobileLauncherProfileWizard/Source/MobileLauncherProfileWizard/Private/SProfileWizardUI.cpp
@@ -43,7 +43,7 @@ namespace ProfileWizardUI
 
 SProfileWizardUI::SProfileWizardUI()
 	: ProfilePlatform(EProfilePlatform::Android)
-	, BuildConfiguration(EBuildConfigurations::Development)
+	, BuildConfiguration(EBuildConfiguration::Development)
 {
 }
 
@@ -489,16 +489,16 @@ TSharedRef<SWidget> SProfileWizardUI::MakeBuildConfigurationMenuContent()
 	struct FConfigInfo
 	{
 		FText ToolTip;
-		EBuildConfigurations::Type Configuration;
+		EBuildConfiguration Configuration;
 	};
 
 	const static FConfigInfo Configurations[] =
 	{
-		{ LOCTEXT("DebugActionHint", "Debug configuration."), EBuildConfigurations::Debug },
-		{ LOCTEXT("DebugGameActionHint", "DebugGame configuration."), EBuildConfigurations::DebugGame },
-		{ LOCTEXT("DevelopmentActionHint", "Development configuration."), EBuildConfigurations::Development },
-		{ LOCTEXT("ShippingActionHint", "Shipping configuration."), EBuildConfigurations::Shipping },
-		{ LOCTEXT("TestActionHint", "Test configuration."), EBuildConfigurations::Test }
+		{ LOCTEXT("DebugActionHint", "Debug configuration."), EBuildConfiguration::Debug },
+		{ LOCTEXT("DebugGameActionHint", "DebugGame configuration."), EBuildConfiguration::DebugGame },
+		{ LOCTEXT("DevelopmentActionHint", "Development configuration."), EBuildConfiguration::Development },
+		{ LOCTEXT("ShippingActionHint", "Shipping configuration."), EBuildConfiguration::Shipping },
+		{ LOCTEXT("TestActionHint", "Test configuration."), EBuildConfiguration::Test }
 	};
 
 	// create build configurations menu
@@ -517,7 +517,7 @@ TSharedRef<SWidget> SProfileWizardUI::MakeBuildConfigurationMenuContent()
 	return MenuBuilder.MakeWidget();
 }
 
-void SProfileWizardUI::HandleBuildConfigurationMenuEntryClicked(EBuildConfigurations::Type InConfiguration)
+void SProfileWizardUI::HandleBuildConfigurationMenuEntryClicked(EBuildConfiguration InConfiguration)
 {
 	BuildConfiguration = InConfiguration;
 }
diff --git a/Engine/Plugins/Editor/MobileLauncherProfileWizard/Source/MobileLauncherProfileWizard/Private/SProfileWizardUI.h b/Engine/Plugins/Editor/MobileLauncherProfileWizard/Source/MobileLauncherProfileWizard/Private/SProfileWizardUI.h
--- a/Engine/Plugins/Editor/MobileLauncherProfileWizard/Source/MobileLauncherProfileWizard/Private/SProfileWizardUI.h
+++ b/Engine/Plugins/Editor/MobileLauncherProfileWizard/Source/MobileLauncherProfileWizard/Private/SProfileWizardUI.h
@@ -37,7 +37,7 @@ namespace EProfilePlatform
 
 struct FProfileParameters
 {
-	EBuildConfigurations::Type BuildConfiguration;
+	EBuildConfiguration BuildConfiguration;
 	FString	ArchiveDirectory;
 	TArray<FString>	AppMaps;
 	TArray<FString>	DLCMaps;
@@ -73,7 +73,7 @@ private:
 	TSharedRef<SWidget> MakeBuildConfigurationMenuContent();
 	
 	/** Handle build configuration selection */
-	void HandleBuildConfigurationMenuEntryClicked(EBuildConfigurations::Type Configuration);
+	void HandleBuildConfigurationMenuEntryClicked(EBuildConfiguration Configuration);
 	
 	/** Get a text for a selected build configuration */
 	FText GetBuildConfigurationSelectorText() const;
@@ -180,7 +180,7 @@ private:
 	FString ProjectPath;
 
 	/** Selected project build configuration */
-	EBuildConfigurations::Type BuildConfiguration;
+	EBuildConfiguration BuildConfiguration;
 
 	/** Cached list of project maps */
 	TArray<TSharedPtr<FString>> ProjectMapList;
diff --git a/Engine/Plugins/Runtime/AR/Microsoft/HoloLensAR/Source/HoloLensTargetPlatform/Private/HoloLensTargetDevice.cpp b/Engine/Plugins/Runtime/AR/Microsoft/HoloLensAR/Source/HoloLensTargetPlatform/Private/HoloLensTargetDevice.cpp
--- a/Engine/Plugins/Runtime/AR/Microsoft/HoloLensAR/Source/HoloLensTargetPlatform/Private/HoloLensTargetDevice.cpp
+++ b/Engine/Plugins/Runtime/AR/Microsoft/HoloLensAR/Source/HoloLensTargetPlatform/Private/HoloLensTargetDevice.cpp
@@ -18,7 +18,7 @@ bool FHoloLensTargetDevice::Deploy(const FString& SourceFolder, FString& OutAppI
 	return Info.bCanDeployTo;
 }
 
-bool FHoloLensTargetDevice::Launch(const FString& AppId, EBuildConfigurations::Type BuildConfiguration, EBuildTargetType BuildTarget, const FString& Params, uint32* OutProcessId)
+bool FHoloLensTargetDevice::Launch(const FString& AppId, EBuildConfiguration BuildConfiguration, EBuildTargetType BuildTarget, const FString& Params, uint32* OutProcessId)
 {
 	return false;
 }
diff --git a/Engine/Plugins/Runtime/AR/Microsoft/HoloLensAR/Source/HoloLensTargetPlatform/Private/HoloLensTargetDevice.h b/Engine/Plugins/Runtime/AR/Microsoft/HoloLensAR/Source/HoloLensTargetPlatform/Private/HoloLensTargetDevice.h
--- a/Engine/Plugins/Runtime/AR/Microsoft/HoloLensAR/Source/HoloLensTargetPlatform/Private/HoloLensTargetDevice.h
+++ b/Engine/Plugins/Runtime/AR/Microsoft/HoloLensAR/Source/HoloLensTargetPlatform/Private/HoloLensTargetDevice.h
@@ -103,7 +103,7 @@ public:
 		return Info.HostName == FPlatformProcess::ComputerName();
 	}
 
-	virtual bool Launch(const FString& AppId, EBuildConfigurations::Type BuildConfiguration, EBuildTargetType TargetType, const FString& Params, uint32* OutProcessId) override;
+	virtual bool Launch(const FString& AppId, EBuildConfiguration BuildConfiguration, EBuildTargetType TargetType, const FString& Params, uint32* OutProcessId) override;
 
 	virtual bool PowerOff(bool Force) override
 	{
diff --git a/Engine/Source/Developer/Android/AndroidTargetPlatform/Private/AndroidTargetDevice.h b/Engine/Source/Developer/Android/AndroidTargetPlatform/Private/AndroidTargetDevice.h
--- a/Engine/Source/Developer/Android/AndroidTargetPlatform/Private/AndroidTargetDevice.h
+++ b/Engine/Source/Developer/Android/AndroidTargetPlatform/Private/AndroidTargetDevice.h
@@ -198,7 +198,7 @@ public:
 		return "Android_Multi";
 	}
 
-	virtual bool Launch(const FString& AppId, EBuildConfigurations::Type BuildConfiguration, EBuildTargetType TargetType, const FString& Params, uint32* OutProcessId);
+	virtual bool Launch(const FString& AppId, EBuildConfiguration BuildConfiguration, EBuildTargetType TargetType, const FString& Params, uint32* OutProcessId);
 	virtual bool Reboot(bool bReconnect = false) override;
 	virtual bool Run(const FString& ExecutablePath, const FString& Params, uint32* OutProcessId) override;
 	virtual bool TerminateLaunchedProcess(const FString& ProcessIdentifier) override;
diff --git a/Engine/Source/Developer/Android/AndroidTargetPlatform/Private/AndroidTargetDevice.inl b/Engine/Source/Developer/Android/AndroidTargetPlatform/Private/AndroidTargetDevice.inl
--- a/Engine/Source/Developer/Android/AndroidTargetPlatform/Private/AndroidTargetDevice.inl
+++ b/Engine/Source/Developer/Android/AndroidTargetPlatform/Private/AndroidTargetDevice.inl
@@ -48,7 +48,7 @@ inline int32 FAndroidTargetDevice::GetProcessSnapshot( TArray<FTargetDeviceProce
 }
 
 
-inline bool FAndroidTargetDevice::Launch( const FString& AppId, EBuildConfigurations::Type BuildConfiguration, EBuildTargetType TargetType, const FString& Params, uint32* OutProcessId )
+inline bool FAndroidTargetDevice::Launch( const FString& AppId, EBuildConfiguration BuildConfiguration, EBuildTargetType TargetType, const FString& Params, uint32* OutProcessId )
 {
 	// this isn't used, UAT handles it all
 	return false;
diff --git a/Engine/Source/Developer/BlueprintNativeCodeGen/Private/BlueprintNativeCodeGenManifest.cpp b/Engine/Source/Developer/BlueprintNativeCodeGen/Private/BlueprintNativeCodeGenManifest.cpp
--- a/Engine/Source/Developer/BlueprintNativeCodeGen/Private/BlueprintNativeCodeGenManifest.cpp
+++ b/Engine/Source/Developer/BlueprintNativeCodeGen/Private/BlueprintNativeCodeGenManifest.cpp
@@ -328,7 +328,7 @@ FBlueprintNativeCodeGenPaths FBlueprintNativeCodeGenPaths::GetDefaultCodeGenPath
 				static const FName UBTTargetId_Windows = FName(TEXT("Windows"));
 
 				const FName UBTTargetId = (PlatformInfo.UBTTargetId == UBTTargetId_Win32 || PlatformInfo.UBTTargetId == UBTTargetId_Win64) ? UBTTargetId_Windows : PlatformInfo.UBTTargetId;
-				DefaultPluginPath = FPaths::Combine(*DefaultPluginPath, *LexToString(UBTTargetId), *LexToString(PlatformInfo.PlatformType));
+				DefaultPluginPath = FPaths::Combine(*DefaultPluginPath, *LexToString(UBTTargetId), LexToString(PlatformInfo.PlatformType));
 				break;
 			}
 		}
diff --git a/Engine/Source/Developer/DesktopPlatform/Private/PlatformInfo.cpp b/Engine/Source/Developer/DesktopPlatform/Private/PlatformInfo.cpp
--- a/Engine/Source/Developer/DesktopPlatform/Private/PlatformInfo.cpp
+++ b/Engine/Source/Developer/DesktopPlatform/Private/PlatformInfo.cpp
@@ -144,7 +144,10 @@ void ParseDataDrivenPlatformInfo(const TCHAR* Name, const FConfigSection& Sectio
 	FString TargetPlatformCanUseCrashReporterString = GetSectionString(Section, TEXT("bTargetPlatformCanUseCrashReporter"));
 	bool bTargetPlatformCanUseCrashReporter = TargetPlatformCanUseCrashReporterString.IsEmpty() ? true : GetSectionBool(Section, TEXT("bTargetPlatformCanUseCrashReporter"));
 
-	BuildPlatformInfo(Name, TargetPlatformName, FText::FromString(DisplayName), LexFromString(*PlatformType), ConvertPlatformFlags(PlatformFlags), FPlatformIconPaths(NormalIconPath, LargeIconPath, XLargeIconPath), UATCommandLine,
+	EBuildTargetType TargetType;
+	LexFromString(TargetType, *PlatformType);
+
+	BuildPlatformInfo(Name, TargetPlatformName, FText::FromString(DisplayName), TargetType, ConvertPlatformFlags(PlatformFlags), FPlatformIconPaths(NormalIconPath, LargeIconPath, XLargeIconPath), UATCommandLine,
 		AutoSDKPath, PlatformInfo::EPlatformSDKStatus::Unknown, TutorialPath, bIsEnabled, BinariesDirectoryName, IniPlatformName, bUsesHostCompiler, bUATClosesAfterLaunch, bIsConfidential, UBTTargetID, PlatformGroupName, bTargetPlatformCanUseCrashReporter);
 }
 
diff --git a/Engine/Source/Developer/HTML5/HTML5TargetPlatform/Private/HTML5TargetDevice.cpp b/Engine/Source/Developer/HTML5/HTML5TargetPlatform/Private/HTML5TargetDevice.cpp
--- a/Engine/Source/Developer/HTML5/HTML5TargetPlatform/Private/HTML5TargetDevice.cpp
+++ b/Engine/Source/Developer/HTML5/HTML5TargetPlatform/Private/HTML5TargetDevice.cpp
@@ -48,7 +48,7 @@ int32 FHTML5TargetDevice::GetProcessSnapshot( TArray<FTargetDeviceProcessInfo>&
 }
 
 
-bool FHTML5TargetDevice::Launch( const FString& AppId, EBuildConfigurations::Type BuildConfiguration, EBuildTargetType TargetType, const FString& Params, uint32* OutProcessId )
+bool FHTML5TargetDevice::Launch( const FString& AppId, EBuildConfiguration BuildConfiguration, EBuildTargetType TargetType, const FString& Params, uint32* OutProcessId )
 {
 	return Run(AppId, Params, OutProcessId);
 }
diff --git a/Engine/Source/Developer/HTML5/HTML5TargetPlatform/Private/HTML5TargetDevice.h b/Engine/Source/Developer/HTML5/HTML5TargetPlatform/Private/HTML5TargetDevice.h
--- a/Engine/Source/Developer/HTML5/HTML5TargetPlatform/Private/HTML5TargetDevice.h
+++ b/Engine/Source/Developer/HTML5/HTML5TargetPlatform/Private/HTML5TargetDevice.h
@@ -90,7 +90,7 @@ public:
 		return true;
 	}
 
-	virtual bool Launch( const FString& AppId, EBuildConfigurations::Type BuildConfiguration, EBuildTargetType TargetType, const FString& Params, uint32* OutProcessId ) override;
+	virtual bool Launch( const FString& AppId, EBuildConfiguration BuildConfiguration, EBuildTargetType TargetType, const FString& Params, uint32* OutProcessId ) override;
 
 	virtual bool PowerOff( bool Force ) override;
 
diff --git a/Engine/Source/Developer/IOS/IOSTargetPlatform/Private/IOSTargetDevice.cpp b/Engine/Source/Developer/IOS/IOSTargetPlatform/Private/IOSTargetDevice.cpp
--- a/Engine/Source/Developer/IOS/IOSTargetPlatform/Private/IOSTargetDevice.cpp
+++ b/Engine/Source/Developer/IOS/IOSTargetPlatform/Private/IOSTargetDevice.cpp
@@ -329,7 +329,7 @@ bool FIOSTargetDevice::IsDefault() const
 	return true;
 }
 	
-bool FIOSTargetDevice::Launch(const FString& InAppId, EBuildConfigurations::Type InBuildConfiguration, EBuildTargetType TargetType, const FString& Params, uint32* OutProcessId)
+bool FIOSTargetDevice::Launch(const FString& InAppId, EBuildConfiguration InBuildConfiguration, EBuildTargetType TargetType, const FString& Params, uint32* OutProcessId)
 {
 #if !PLATFORM_MAC
 	MessageEndpoint->Send(new FIOSLaunchDaemonLaunchApp(InAppId, Params), DeviceEndpoint);
diff --git a/Engine/Source/Developer/IOS/IOSTargetPlatform/Private/IOSTargetDevice.h b/Engine/Source/Developer/IOS/IOSTargetPlatform/Private/IOSTargetDevice.h
--- a/Engine/Source/Developer/IOS/IOSTargetPlatform/Private/IOSTargetDevice.h
+++ b/Engine/Source/Developer/IOS/IOSTargetPlatform/Private/IOSTargetDevice.h
@@ -118,7 +118,7 @@ public:
 	virtual const class ITargetPlatform& GetTargetPlatform() const override;
 	virtual bool IsConnected() override;
 	virtual bool IsDefault() const override;
-	virtual bool Launch(const FString& InAppId, EBuildConfigurations::Type InBuildConfiguration, EBuildTargetType TargetType, const FString& Params, uint32* OutProcessId) override;
+	virtual bool Launch(const FString& InAppId, EBuildConfiguration InBuildConfiguration, EBuildTargetType TargetType, const FString& Params, uint32* OutProcessId) override;
 	virtual bool PowerOff(bool Force) override;
 	virtual bool PowerOn() override;
 	virtual bool Reboot(bool bReconnect = false) override;
@@ -154,7 +154,7 @@ private:
 	FString AppId;
 
 	/** Contains the build configuration of the app to deploy */
-	EBuildConfigurations::Type BuildConfiguration;
+	EBuildConfiguration BuildConfiguration;
 
 	/** Lets us know whether the thing is a sim device or a physical device. */
 	bool bIsSimulated;
@@ -244,7 +244,7 @@ public:
 		AppId = GameName;
 	}
 
-	void SetAppConfiguration(EBuildConfigurations::Type Configuration)
+	void SetAppConfiguration(EBuildConfiguration Configuration)
 	{
 		BuildConfiguration = Configuration;
 	}
diff --git a/Engine/Source/Developer/LauncherServices/Private/Launcher/LauncherUATTask.h b/Engine/Source/Developer/LauncherServices/Private/Launcher/LauncherUATTask.h
--- a/Engine/Source/Developer/LauncherServices/Private/Launcher/LauncherUATTask.h
+++ b/Engine/Source/Developer/LauncherServices/Private/Launcher/LauncherUATTask.h
@@ -58,8 +58,8 @@ protected:
 		UATCommandLine = FString::Printf(TEXT("-ScriptsForProject=\"%s\" BuildCookRun -project=\"%s\" -noP4 -clientconfig=%s -serverconfig=%s"),
 			*ProjectPath,
 			*ProjectPath,
-			*ConfigStrings[ChainState.Profile->GetBuildConfiguration()],
-			*ConfigStrings[ChainState.Profile->GetBuildConfiguration()]);
+			LexToString(ChainState.Profile->GetBuildConfiguration()),
+			LexToString(ChainState.Profile->GetBuildConfiguration()));
 
 		// we expect to pass -nocompile to UAT here as we generally expect UAT to be fully compiled. Besides, installed builds don't even have the source to compile UAT scripts.
 		// Only allow UAT to compile scripts dynamically if we pass -development or we have the IsBuildingUAT property set, the latter of which should not allow itself to be set in installed situations.
diff --git a/Engine/Source/Developer/LauncherServices/Private/Profiles/LauncherProfile.h b/Engine/Source/Developer/LauncherServices/Private/Profiles/LauncherProfile.h
--- a/Engine/Source/Developer/LauncherServices/Private/Profiles/LauncherProfile.h
+++ b/Engine/Source/Developer/LauncherServices/Private/Profiles/LauncherProfile.h
@@ -82,7 +82,7 @@ public:
 		return Variant;
 	}
 
-	virtual EBuildConfigurations::Type GetBuildConfiguration() const override
+	virtual EBuildConfiguration GetBuildConfiguration() const override
 	{
 		return BuildConfiguration;
 	}
@@ -105,7 +105,7 @@ public:
 		Variant = InVariant;
 	}
 
-	virtual void SetBuildConfiguration(EBuildConfigurations::Type InConfiguration) override
+	virtual void SetBuildConfiguration(EBuildConfiguration InConfiguration) override
 	{
 		BuildConfiguration = InConfiguration;
 	}
@@ -143,7 +143,7 @@ public:
 		Writer.WriteValue("Version", Version);
 		Writer.WriteValue("DeviceName", DeviceName);
 		Writer.WriteValue("Variant", Variant.ToString());
-		Writer.WriteValue("BuildConfiguration", BuildConfiguration);
+		Writer.WriteValue("BuildConfiguration", (int32)BuildConfiguration);
 		Writer.WriteValue("CookMode", CookMode);
 		Writer.WriteObjectEnd();
 	}
@@ -158,7 +158,7 @@ public:
 
 		DeviceName = Object.GetStringField("DeviceName");
 		Variant = *(Object.GetStringField("Variant"));
-		BuildConfiguration = (TEnumAsByte<EBuildConfigurations::Type>)((int32)Object.GetNumberField("BuildConfiguration"));
+		BuildConfiguration = (EBuildConfiguration)((int32)Object.GetNumberField("BuildConfiguration"));
 		CookMode = (TEnumAsByte<ELauncherProfileCookModes::Type>)((int32)Object.GetNumberField("CookMode"));
 
 		return true;
@@ -171,7 +171,7 @@ public:
 		
 		// I don't use FApp::GetBuildConfiguration() because i don't want the act of running in debug the first time to cause the simple
 		// profiles created for your persistent devices to be in debug. The user might not see this if they don't expand the Advanced options.
-		BuildConfiguration = EBuildConfigurations::Development;
+		BuildConfiguration = EBuildConfiguration::Development;
 		
 		CookMode = ELauncherProfileCookModes::OnTheFly;		
 	}
@@ -185,7 +185,7 @@ private:
 	FName Variant;
 
 	// Holds the desired build configuration (only used if creating new builds).
-	TEnumAsByte<EBuildConfigurations::Type> BuildConfiguration;
+	EBuildConfiguration BuildConfiguration;
 
 	// Holds the cooking mode.
 	TEnumAsByte<ELauncherProfileCookModes::Type> CookMode;
@@ -363,12 +363,12 @@ public:
 		return Role;
 	}
 
-	virtual EBuildConfigurations::Type GetBuildConfiguration( ) const override
+	virtual EBuildConfiguration GetBuildConfiguration( ) const override
 	{
 		return BuildConfiguration;
 	}
 
-	virtual EBuildConfigurations::Type GetCookConfiguration( ) const override
+	virtual EBuildConfiguration GetCookConfiguration( ) const override
 	{
 		return CookConfiguration;
 	}
@@ -988,10 +988,10 @@ public:
 		Writer.WriteValue("Id", Id.ToString());
 		Writer.WriteValue("Name", Name);
 		Writer.WriteValue("Description", Description);
-		Writer.WriteValue("BuildConfiguration", BuildConfiguration);
+		Writer.WriteValue("BuildConfiguration", (int32)BuildConfiguration);
 		Writer.WriteValue("ProjectSpecified", ProjectSpecified);
 		Writer.WriteValue("ShareableProjectPath", ShareableProjectPath);
-		Writer.WriteValue("CookConfiguration", CookConfiguration);
+		Writer.WriteValue("CookConfiguration", (int32)CookConfiguration);
 		Writer.WriteValue("CookIncremental", CookIncremental);
 		Writer.WriteValue("CookOptions", CookOptions);
 		Writer.WriteValue("CookMode", CookMode);
@@ -1106,14 +1106,13 @@ public:
 		Writer.WriteValue("utf8output", true);
 
 		// client configurations
-		static const FString ConfigStrings[] = { TEXT("Unknown"), TEXT("Debug"), TEXT("DebugGame"), TEXT("Development"), TEXT("Shipping"), TEXT("Test") };
 		Writer.WriteArrayStart("clientconfig");
-		Writer.WriteValue(ConfigStrings[BuildConfiguration]);
+		Writer.WriteValue(LexToString(BuildConfiguration));
 		Writer.WriteArrayEnd();
 
 		// server configurations
 		Writer.WriteArrayStart("serverconfig");
-		Writer.WriteValue(ConfigStrings[BuildConfiguration]);
+		Writer.WriteValue(LexToString(BuildConfiguration));
 		Writer.WriteArrayEnd();
 
 		// platforms
@@ -1614,10 +1613,10 @@ public:
 		FGuid::Parse(Object.GetStringField("Id"), Id);
 		Name = Object.GetStringField("Name");
 		Description = Object.GetStringField("Description");
-		BuildConfiguration = (TEnumAsByte<EBuildConfigurations::Type>)((int32)Object.GetNumberField("BuildConfiguration"));
+		BuildConfiguration = (EBuildConfiguration)((int32)Object.GetNumberField("BuildConfiguration"));
 		ProjectSpecified = Object.GetBoolField("ProjectSpecified");
 		ShareableProjectPath = Object.GetStringField("ShareableProjectPath");
-		CookConfiguration = (TEnumAsByte<EBuildConfigurations::Type>)((int32)Object.GetNumberField("CookConfiguration"));
+		CookConfiguration = (EBuildConfiguration)((int32)Object.GetNumberField("CookConfiguration"));
 		CookIncremental = Object.GetBoolField("CookIncremental");
 		CookOptions = Object.GetStringField("CookOptions");
 		CookMode = (TEnumAsByte<ELauncherProfileCookModes::Type>)((int32)Object.GetNumberField("CookMode"));
@@ -1785,7 +1784,7 @@ public:
 		
 		// I don't use FApp::GetBuildConfiguration() because i don't want the act of running in debug the first time to cause 
 		// profiles the user creates to be in debug. This will keep consistency.
-		BuildConfiguration = EBuildConfigurations::Development;
+		BuildConfiguration = EBuildConfiguration::Development;
 
 		FInternationalization& I18N = FInternationalization::Get();
 
@@ -1895,7 +1894,7 @@ public:
 		}
 	}
 
-	virtual void SetBuildConfiguration( EBuildConfigurations::Type Configuration ) override
+	virtual void SetBuildConfiguration( EBuildConfiguration Configuration ) override
 	{
 		if (BuildConfiguration != Configuration)
 		{
@@ -1905,7 +1904,7 @@ public:
 		}
 	}
 
-	virtual void SetCookConfiguration( EBuildConfigurations::Type Configuration ) override
+	virtual void SetCookConfiguration( EBuildConfiguration Configuration ) override
 	{
 		if (CookConfiguration != Configuration)
 		{
@@ -2341,7 +2340,7 @@ protected:
 		ValidationErrors.Reset();
 
 		// Build: a build configuration must be selected
-		if (BuildConfiguration == EBuildConfigurations::Unknown)
+		if (BuildConfiguration == EBuildConfiguration::Unknown)
 		{
 			ValidationErrors.Add(ELauncherProfileValidationErrors::NoBuildConfigurationSelected);
 		}
@@ -2451,7 +2450,7 @@ protected:
 		
 		if ((CookMode == ELauncherProfileCookModes::OnTheFly) || (CookMode == ELauncherProfileCookModes::OnTheFlyInEditor))
 		{
-			if (BuildConfiguration == EBuildConfigurations::Shipping)
+			if (BuildConfiguration == EBuildConfiguration::Shipping)
 			{
 				// shipping doesn't support commandline options
 				ValidationErrors.Add(ELauncherProfileValidationErrors::ShippingDoesntSupportCommandlineOptionsCantUseCookOnTheFly);
@@ -2594,11 +2593,11 @@ private:
 	ILauncherProfileManagerRef LauncherProfileManager;
 
 	// Holds the desired build configuration (only used if creating new builds).
-	TEnumAsByte<EBuildConfigurations::Type> BuildConfiguration;
+	EBuildConfiguration BuildConfiguration;
 
 	// Holds the build mode.
 	// Holds the build configuration name of the cooker.
-	TEnumAsByte<EBuildConfigurations::Type> CookConfiguration;
+	EBuildConfiguration CookConfiguration;
 
 	// Holds additional cooker command line options.
 	FString CookOptions;
diff --git a/Engine/Source/Developer/LauncherServices/Public/ILauncherProfile.h b/Engine/Source/Developer/LauncherServices/Public/ILauncherProfile.h
--- a/Engine/Source/Developer/LauncherServices/Public/ILauncherProfile.h
+++ b/Engine/Source/Developer/LauncherServices/Public/ILauncherProfile.h
@@ -241,7 +241,7 @@ public:
 	 * @return Build configuration name.
 	 * @see SetBuildConfigurationName
 	 */
-	virtual EBuildConfigurations::Type GetBuildConfiguration() const = 0;
+	virtual EBuildConfiguration GetBuildConfiguration() const = 0;
 
 	/**
 	 * Gets the selected cook mode.
@@ -285,7 +285,7 @@ public:
 	 * @param InConfiguration The build configuration name to set.
 	 * @see GetBuildConfigurationName
 	 */
-	virtual void SetBuildConfiguration(EBuildConfigurations::Type InConfiguration) = 0;
+	virtual void SetBuildConfiguration(EBuildConfiguration InConfiguration) = 0;
 
 	/**
 	 * Sets the cook mode.
@@ -487,7 +487,7 @@ public:
 	 * @return Build configuration name.
 	 * @see SetBuildConfigurationName
 	 */
-	virtual EBuildConfigurations::Type GetBuildConfiguration( ) const = 0;
+	virtual EBuildConfiguration GetBuildConfiguration( ) const = 0;
 
 	/**
 	 * Gets the build configuration name of the cooker.
@@ -495,7 +495,7 @@ public:
 	 * @return Cook configuration name.
 	 * @see SetCookConfigurationName
 	 */
-	virtual EBuildConfigurations::Type GetCookConfiguration( ) const = 0;
+	virtual EBuildConfiguration GetCookConfiguration( ) const = 0;
 
 	/**
 	 * Gets the selected cook mode.
@@ -983,7 +983,7 @@ public:
 	 * @param ConfigurationName The build configuration name to set.
 	 * @see GetBuildConfigurationName
 	 */
-	virtual void SetBuildConfiguration( EBuildConfigurations::Type Configuration ) = 0;
+	virtual void SetBuildConfiguration( EBuildConfiguration Configuration ) = 0;
 
 	/**
 	 * Sets the build configuration of the cooker.
@@ -991,7 +991,7 @@ public:
 	 * @param Configuration The cooker's build configuration to set.
 	 * @see GetBuildConfigurationName
 	 */
-	virtual void SetCookConfiguration( EBuildConfigurations::Type Configuration ) = 0;
+	virtual void SetCookConfiguration( EBuildConfiguration Configuration ) = 0;
 
 	/**
 	 * Sets the cook mode.
diff --git a/Engine/Source/Developer/Linux/LinuxTargetPlatform/Private/LinuxTargetDevice.h b/Engine/Source/Developer/Linux/LinuxTargetPlatform/Private/LinuxTargetDevice.h
--- a/Engine/Source/Developer/Linux/LinuxTargetPlatform/Private/LinuxTargetDevice.h
+++ b/Engine/Source/Developer/Linux/LinuxTargetPlatform/Private/LinuxTargetDevice.h
@@ -143,7 +143,7 @@ public:
 		return false;
 	}
 
-	virtual bool Launch( const FString& AppId, EBuildConfigurations::Type BuildConfiguration, EBuildTargetType TargetType, const FString& Params, uint32* OutProcessId ) override
+	virtual bool Launch( const FString& AppId, EBuildConfiguration BuildConfiguration, EBuildTargetType TargetType, const FString& Params, uint32* OutProcessId ) override
 	{
 #if PLATFORM_LINUX	// if running natively, support launching in place
 		// build executable path
@@ -163,9 +163,9 @@ public:
 			ExecutablePath /= TEXT("UE4Editor");
 		}
 
-		if (BuildConfiguration != EBuildConfigurations::Development)
+		if (BuildConfiguration != EBuildConfiguration::Development)
 		{
-			ExecutablePath += FString::Printf(TEXT("-%s-%s"), *PlatformName, EBuildConfigurations::ToString(BuildConfiguration));
+			ExecutablePath += FString::Printf(TEXT("-%s-%s"), *PlatformName, LexToString(BuildConfiguration));
 		}
 
 		// launch the game
diff --git a/Engine/Source/Developer/Mac/MacTargetPlatform/Private/LocalMacTargetDevice.h b/Engine/Source/Developer/Mac/MacTargetPlatform/Private/LocalMacTargetDevice.h
--- a/Engine/Source/Developer/Mac/MacTargetPlatform/Private/LocalMacTargetDevice.h
+++ b/Engine/Source/Developer/Mac/MacTargetPlatform/Private/LocalMacTargetDevice.h
@@ -106,14 +106,14 @@ public:
 		return true;
 	}
 
-	virtual bool Launch( const FString& AppId, EBuildConfigurations::Type BuildConfiguration, EBuildTargetType TargetType, const FString& Params, uint32* OutProcessId ) override
+	virtual bool Launch( const FString& AppId, EBuildConfiguration BuildConfiguration, EBuildTargetType TargetType, const FString& Params, uint32* OutProcessId ) override
 	{
 		// build executable path
 		FString PlatformName = TEXT("Mac");
 		FString ExecutableName = TEXT("UE4");
-		if (BuildConfiguration != EBuildConfigurations::Development)
+		if (BuildConfiguration != EBuildConfiguration::Development)
 		{
-			ExecutableName += FString::Printf(TEXT("-%s-%s"), *PlatformName, EBuildConfigurations::ToString(BuildConfiguration));
+			ExecutableName += FString::Printf(TEXT("-%s-%s"), *PlatformName, LexToString(BuildConfiguration));
 		}
 
 		FString ExecutablePath = FPaths::EngineIntermediateDir() / TEXT("Devices") / PlatformName / TEXT("Engine") / TEXT("Binaries") / PlatformName / (ExecutableName + TEXT(".app/Contents/MacOS/") + ExecutableName);
diff --git a/Engine/Source/Developer/ProjectLauncher/Private/Widgets/Build/SProjectLauncherBuildPage.cpp b/Engine/Source/Developer/ProjectLauncher/Private/Widgets/Build/SProjectLauncherBuildPage.cpp
--- a/Engine/Source/Developer/ProjectLauncher/Private/Widgets/Build/SProjectLauncherBuildPage.cpp
+++ b/Engine/Source/Developer/ProjectLauncher/Private/Widgets/Build/SProjectLauncherBuildPage.cpp
@@ -261,9 +261,9 @@ FReply SProjectLauncherBuildPage::HandleGenDSYMClicked()
         if (!SelectedProfile->HasValidationError(ELauncherProfileValidationErrors::NoProjectSelected))
         {
             FString ProjectName = SelectedProfile->GetProjectName();
-            EBuildConfigurations::Type ProjectConfig = SelectedProfile->GetBuildConfiguration();
+            EBuildConfiguration ProjectConfig = SelectedProfile->GetBuildConfiguration();
 
-            GenerateDSYMForProject(ProjectName, EBuildConfigurations::ToString(ProjectConfig));
+            GenerateDSYMForProject(ProjectName, LexToString(ProjectConfig));
         }
     }
 
@@ -302,7 +302,7 @@ EVisibility SProjectLauncherBuildPage::ShowBuildConfiguration() const
 }
 
 
-void SProjectLauncherBuildPage::HandleBuildConfigurationSelectorConfigurationSelected(EBuildConfigurations::Type Configuration)
+void SProjectLauncherBuildPage::HandleBuildConfigurationSelectorConfigurationSelected(EBuildConfiguration Configuration)
 {
 	ILauncherProfilePtr SelectedProfile = Model->GetSelectedProfile();
 
@@ -319,7 +319,7 @@ FText SProjectLauncherBuildPage::HandleBuildConfigurationSelectorText() const
 
 	if (SelectedProfile.IsValid())
 	{
-		return FText::FromString(EBuildConfigurations::ToString(SelectedProfile->GetBuildConfiguration()));
+		return FText::FromString(LexToString(SelectedProfile->GetBuildConfiguration()));
 	}
 
 	return FText::GetEmpty();
diff --git a/Engine/Source/Developer/ProjectLauncher/Private/Widgets/Build/SProjectLauncherBuildPage.h b/Engine/Source/Developer/ProjectLauncher/Private/Widgets/Build/SProjectLauncherBuildPage.h
--- a/Engine/Source/Developer/ProjectLauncher/Private/Widgets/Build/SProjectLauncherBuildPage.h
+++ b/Engine/Source/Developer/ProjectLauncher/Private/Widgets/Build/SProjectLauncherBuildPage.h
@@ -73,7 +73,7 @@ private:
 	EVisibility ShowBuildConfiguration() const;
 
 	/** Callback for selecting a build configuration. */
-	void HandleBuildConfigurationSelectorConfigurationSelected(EBuildConfigurations::Type Configuration);
+	void HandleBuildConfigurationSelectorConfigurationSelected(EBuildConfiguration Configuration);
 
 	/** Callback for getting the content text of the build configuration selector. */
 	FText HandleBuildConfigurationSelectorText() const;
diff --git a/Engine/Source/Developer/ProjectLauncher/Private/Widgets/Cook/SProjectLauncherCookByTheBookSettings.cpp b/Engine/Source/Developer/ProjectLauncher/Private/Widgets/Cook/SProjectLauncherCookByTheBookSettings.cpp
--- a/Engine/Source/Developer/ProjectLauncher/Private/Widgets/Cook/SProjectLauncherCookByTheBookSettings.cpp
+++ b/Engine/Source/Developer/ProjectLauncher/Private/Widgets/Cook/SProjectLauncherCookByTheBookSettings.cpp
@@ -1299,7 +1299,7 @@ EVisibility SProjectLauncherCookByTheBookSettings::HandleMapSelectionHyperlinkVi
 }
 
 
-void SProjectLauncherCookByTheBookSettings::HandleCookConfigurationSelectorConfigurationSelected(EBuildConfigurations::Type Configuration)
+void SProjectLauncherCookByTheBookSettings::HandleCookConfigurationSelectorConfigurationSelected(EBuildConfiguration Configuration)
 {
 	ILauncherProfilePtr SelectedProfile = Model->GetSelectedProfile();
 
@@ -1316,7 +1316,7 @@ FText SProjectLauncherCookByTheBookSettings::HandleCookConfigurationSelectorText
 
 	if (SelectedProfile.IsValid())
 	{
-		return FText::FromString(EBuildConfigurations::ToString(SelectedProfile->GetCookConfiguration()));
+		return FText::FromString(LexToString(SelectedProfile->GetCookConfiguration()));
 	}
 
 	return FText::GetEmpty();
diff --git a/Engine/Source/Developer/ProjectLauncher/Private/Widgets/Cook/SProjectLauncherCookByTheBookSettings.h b/Engine/Source/Developer/ProjectLauncher/Private/Widgets/Cook/SProjectLauncherCookByTheBookSettings.h
--- a/Engine/Source/Developer/ProjectLauncher/Private/Widgets/Cook/SProjectLauncherCookByTheBookSettings.h
+++ b/Engine/Source/Developer/ProjectLauncher/Private/Widgets/Cook/SProjectLauncherCookByTheBookSettings.h
@@ -74,7 +74,7 @@ private:
 	void HandleAllMapsHyperlinkNavigate(bool AllPlatforms);
 
 	/** Handles selecting a build configuration for the cooker. */
-	void HandleCookConfigurationSelectorConfigurationSelected(EBuildConfigurations::Type);
+	void HandleCookConfigurationSelectorConfigurationSelected(EBuildConfiguration);
 
 	/** Handles getting the content text of the cooker build configuration selector. */
 	FText HandleCookConfigurationSelectorText() const;
diff --git a/Engine/Source/Developer/ProjectLauncher/Private/Widgets/Deploy/SProjectLauncherSimpleDeviceListRow.h b/Engine/Source/Developer/ProjectLauncher/Private/Widgets/Deploy/SProjectLauncherSimpleDeviceListRow.h
--- a/Engine/Source/Developer/ProjectLauncher/Private/Widgets/Deploy/SProjectLauncherSimpleDeviceListRow.h
+++ b/Engine/Source/Developer/ProjectLauncher/Private/Widgets/Deploy/SProjectLauncherSimpleDeviceListRow.h
@@ -291,7 +291,7 @@ private:
 		return FText::GetEmpty();
 	}
 
-	void HandleBuildConfigurationSelectorConfigurationSelected(EBuildConfigurations::Type Configuration)
+	void HandleBuildConfigurationSelectorConfigurationSelected(EBuildConfiguration Configuration)
 	{
 		if (SimpleProfile.IsValid())
 		{
diff --git a/Engine/Source/Developer/ProjectLauncher/Private/Widgets/Preview/SProjectLauncherPreviewPage.cpp b/Engine/Source/Developer/ProjectLauncher/Private/Widgets/Preview/SProjectLauncherPreviewPage.cpp
--- a/Engine/Source/Developer/ProjectLauncher/Private/Widgets/Preview/SProjectLauncherPreviewPage.cpp
+++ b/Engine/Source/Developer/ProjectLauncher/Private/Widgets/Preview/SProjectLauncherPreviewPage.cpp
@@ -845,7 +845,7 @@ FText SProjectLauncherPreviewPage::HandleBuildConfigurationTextBlockText( ) cons
 
 	if (SelectedProfile.IsValid())
 	{
-		return FText::FromString(EBuildConfigurations::ToString(SelectedProfile->GetBuildConfiguration()));
+		return FText::FromString(LexToString(SelectedProfile->GetBuildConfiguration()));
 	}
 
 	return LOCTEXT("NotAvailableText", "n/a");
diff --git a/Engine/Source/Developer/ProjectLauncher/Private/Widgets/Shared/SProjectLauncherBuildConfigurationSelector.h b/Engine/Source/Developer/ProjectLauncher/Private/Widgets/Shared/SProjectLauncherBuildConfigurationSelector.h
--- a/Engine/Source/Developer/ProjectLauncher/Private/Widgets/Shared/SProjectLauncherBuildConfigurationSelector.h
+++ b/Engine/Source/Developer/ProjectLauncher/Private/Widgets/Shared/SProjectLauncherBuildConfigurationSelector.h
@@ -23,7 +23,7 @@
  *
  * The first parameter is the selected build configuration.
  */
-DECLARE_DELEGATE_OneParam(FOnSessionSProjectLauncherBuildConfigurationSelected, EBuildConfigurations::Type)
+DECLARE_DELEGATE_OneParam(FOnSessionSProjectLauncherBuildConfigurationSelected, EBuildConfiguration)
 
 
 /**
@@ -55,16 +55,16 @@ public:
 		struct FConfigInfo
 		{
 			FText ToolTip;
-			EBuildConfigurations::Type Configuration;
+			EBuildConfiguration Configuration;
 		};
 
 		const FConfigInfo Configurations[] =
 		{
-			{LOCTEXT("DebugActionHint", "Debug configuration."), EBuildConfigurations::Debug},
-			{LOCTEXT("DebugGameActionHint", "DebugGame configuration."), EBuildConfigurations::DebugGame},
-			{LOCTEXT("DevelopmentActionHint", "Development configuration."), EBuildConfigurations::Development},
-			{LOCTEXT("ShippingActionHint", "Shipping configuration."), EBuildConfigurations::Shipping},
-			{LOCTEXT("TestActionHint", "Test configuration."), EBuildConfigurations::Test}
+			{LOCTEXT("DebugActionHint", "Debug configuration."), EBuildConfiguration::Debug},
+			{LOCTEXT("DebugGameActionHint", "DebugGame configuration."), EBuildConfiguration::DebugGame},
+			{LOCTEXT("DevelopmentActionHint", "Development configuration."), EBuildConfiguration::Development},
+			{LOCTEXT("ShippingActionHint", "Shipping configuration."), EBuildConfiguration::Shipping},
+			{LOCTEXT("TestActionHint", "Test configuration."), EBuildConfiguration::Test}
 		};
 
 		// create build configurations menu
@@ -104,7 +104,7 @@ public:
 private:
 
 	// Callback for clicking a menu entry.
-	void HandleMenuEntryClicked( EBuildConfigurations::Type Configuration )
+	void HandleMenuEntryClicked( EBuildConfiguration Configuration )
 	{
 		OnConfigurationSelected.ExecuteIfBound(Configuration);
 	}
diff --git a/Engine/Source/Developer/TargetDeviceServices/Private/Proxies/TargetDeviceProxy.cpp b/Engine/Source/Developer/TargetDeviceServices/Private/Proxies/TargetDeviceProxy.cpp
--- a/Engine/Source/Developer/TargetDeviceServices/Private/Proxies/TargetDeviceProxy.cpp
+++ b/Engine/Source/Developer/TargetDeviceServices/Private/Proxies/TargetDeviceProxy.cpp
@@ -275,7 +275,7 @@ bool FTargetDeviceProxy::DeployApp(FName InVariant, const TMap<FString, FString>
 }
 
 
-bool FTargetDeviceProxy::LaunchApp(FName InVariant, const FString& AppId, EBuildConfigurations::Type BuildConfiguration, const FString& Params)
+bool FTargetDeviceProxy::LaunchApp(FName InVariant, const FString& AppId, EBuildConfiguration BuildConfiguration, const FString& Params)
 {
 	MessageEndpoint->Send(new FTargetDeviceServiceLaunchApp(InVariant, AppId, BuildConfiguration, Params), MessageAddress);
 
diff --git a/Engine/Source/Developer/TargetDeviceServices/Private/Proxies/TargetDeviceProxy.h b/Engine/Source/Developer/TargetDeviceServices/Private/Proxies/TargetDeviceProxy.h
--- a/Engine/Source/Developer/TargetDeviceServices/Private/Proxies/TargetDeviceProxy.h
+++ b/Engine/Source/Developer/TargetDeviceServices/Private/Proxies/TargetDeviceProxy.h
@@ -167,7 +167,7 @@ public:
 	}
 
 	virtual bool DeployApp(FName InVariant, const TMap<FString, FString>& Files, const FGuid& TransactionId) override;
-	virtual bool LaunchApp(FName InVariant, const FString& AppId, EBuildConfigurations::Type BuildConfiguration, const FString& Params) override;
+	virtual bool LaunchApp(FName InVariant, const FString& AppId, EBuildConfiguration BuildConfiguration, const FString& Params) override;
 	
 	virtual bool TerminateLaunchedProcess(FName InVariant, const FString& ProcessIdentifier) override;
 
diff --git a/Engine/Source/Developer/TargetDeviceServices/Private/Services/TargetDeviceService.cpp b/Engine/Source/Developer/TargetDeviceServices/Private/Services/TargetDeviceService.cpp
--- a/Engine/Source/Developer/TargetDeviceServices/Private/Services/TargetDeviceService.cpp
+++ b/Engine/Source/Developer/TargetDeviceServices/Private/Services/TargetDeviceService.cpp
@@ -463,7 +463,7 @@ void FTargetDeviceService::HandleLaunchAppMessage(const FTargetDeviceServiceLaun
 	if (TargetDevice.IsValid())
 	{
 		uint32 ProcessId;
-		bool Succeeded = TargetDevice->Launch(Message.AppID, (EBuildConfigurations::Type)Message.BuildConfiguration, EBuildTargetType::Game, Message.Params, &ProcessId);
+		bool Succeeded = TargetDevice->Launch(Message.AppID, (EBuildConfiguration)Message.BuildConfiguration, EBuildTargetType::Game, Message.Params, &ProcessId);
 
 		if (MessageEndpoint.IsValid())
 		{
diff --git a/Engine/Source/Developer/TargetDeviceServices/Public/ITargetDeviceProxy.h b/Engine/Source/Developer/TargetDeviceServices/Public/ITargetDeviceProxy.h
--- a/Engine/Source/Developer/TargetDeviceServices/Public/ITargetDeviceProxy.h
+++ b/Engine/Source/Developer/TargetDeviceServices/Public/ITargetDeviceProxy.h
@@ -312,7 +312,7 @@ public:
 	 * @return true on success, false otherwise.
 	 * @see Deploy, OnLaunchFailed
 	 */
-	virtual bool LaunchApp(FName InVariant, const FString& AppId, EBuildConfigurations::Type BuildConfiguration, const FString& Params) = 0;
+	virtual bool LaunchApp(FName InVariant, const FString& AppId, EBuildConfiguration BuildConfiguration, const FString& Params) = 0;
 
 	/**
 	 * Powers off the device.
diff --git a/Engine/Source/Developer/TargetDeviceServices/Public/TargetDeviceServiceMessages.h b/Engine/Source/Developer/TargetDeviceServices/Public/TargetDeviceServiceMessages.h
--- a/Engine/Source/Developer/TargetDeviceServices/Public/TargetDeviceServiceMessages.h
+++ b/Engine/Source/Developer/TargetDeviceServices/Public/TargetDeviceServiceMessages.h
@@ -166,10 +166,10 @@ struct FTargetDeviceServiceLaunchApp
 	FTargetDeviceServiceLaunchApp() : BuildConfiguration(0) { }
 
 	/** Creates and initializes a new instance. */
-	FTargetDeviceServiceLaunchApp(FName InVariant, const FString& InAppId, uint8 InBuildConfiguration, const FString& InParams)
+	FTargetDeviceServiceLaunchApp(FName InVariant, const FString& InAppId, EBuildConfiguration InBuildConfiguration, const FString& InParams)
 		: Variant(InVariant)
 		, AppID(InAppId)
-		, BuildConfiguration(InBuildConfiguration)
+		, BuildConfiguration((uint8)InBuildConfiguration)
 		, Params(InParams)
 	{ }
 };
diff --git a/Engine/Source/Developer/TargetPlatform/Private/InstalledPlatformInfo.cpp b/Engine/Source/Developer/TargetPlatform/Private/InstalledPlatformInfo.cpp
--- a/Engine/Source/Developer/TargetPlatform/Private/InstalledPlatformInfo.cpp
+++ b/Engine/Source/Developer/TargetPlatform/Private/InstalledPlatformInfo.cpp
@@ -57,12 +57,12 @@ void FInstalledPlatformInfo::ParsePlatformConfiguration(FString PlatformConfigur
 	bool bCanCreateEntry = true;
 
 	FString ConfigurationName;
-	EBuildConfigurations::Type Configuration = EBuildConfigurations::Unknown;
+	EBuildConfiguration Configuration = EBuildConfiguration::Unknown;
 	if (FParse::Value(*PlatformConfiguration, TEXT("Configuration="), ConfigurationName))
 	{
-		Configuration = EBuildConfigurations::FromString(ConfigurationName);
+		LexFromString(Configuration, *ConfigurationName);
 	}
-	if (Configuration == EBuildConfigurations::Unknown)
+	if (Configuration == EBuildConfiguration::Unknown)
 	{
 		UE_LOG(LogInstalledPlatforms, Warning, TEXT("Unable to read configuration from %s"), *PlatformConfiguration);
 		bCanCreateEntry = false;
@@ -79,7 +79,7 @@ void FInstalledPlatformInfo::ParsePlatformConfiguration(FString PlatformConfigur
 	EBuildTargetType PlatformType = EBuildTargetType::Game;
 	if (FParse::Value(*PlatformConfiguration, TEXT("PlatformType="), PlatformTypeName))
 	{
-		PlatformType = LexFromString(*PlatformTypeName);
+		LexFromString(PlatformType, *PlatformTypeName);
 	}
 
 	FString Architecture;
@@ -113,7 +113,7 @@ void FInstalledPlatformInfo::ParsePlatformConfiguration(FString PlatformConfigur
 	}
 }
 
-bool FInstalledPlatformInfo::IsValidConfiguration(const EBuildConfigurations::Type Configuration, EProjectType ProjectType) const
+bool FInstalledPlatformInfo::IsValidConfiguration(const EBuildConfiguration Configuration, EProjectType ProjectType) const
 {
 	return ContainsValidConfiguration(
 		[Configuration, ProjectType](const FInstalledPlatformConfiguration& CurConfig)
@@ -137,7 +137,7 @@ bool FInstalledPlatformInfo::IsValidPlatform(const FString& PlatformName, EProje
 	);
 }
 
-bool FInstalledPlatformInfo::IsValidPlatformAndConfiguration(const EBuildConfigurations::Type Configuration, const FString& PlatformName, EProjectType ProjectType) const
+bool FInstalledPlatformInfo::IsValidPlatformAndConfiguration(const EBuildConfiguration Configuration, const FString& PlatformName, EProjectType ProjectType) const
 {
 	return ContainsValidConfiguration(
 		[Configuration, PlatformName, ProjectType](const FInstalledPlatformConfiguration& CurConfig)
diff --git a/Engine/Source/Developer/TargetPlatform/Public/InstalledPlatformInfo.h b/Engine/Source/Developer/TargetPlatform/Public/InstalledPlatformInfo.h
--- a/Engine/Source/Developer/TargetPlatform/Public/InstalledPlatformInfo.h
+++ b/Engine/Source/Developer/TargetPlatform/Public/InstalledPlatformInfo.h
@@ -21,7 +21,7 @@ EProjectType EProjectTypeFromString(const FString& ProjectTypeName);
 struct FInstalledPlatformConfiguration
 {
 	/** Build Configuration of this combination */
-	EBuildConfigurations::Type Configuration;
+	EBuildConfiguration Configuration;
 
 	/** Name of the Platform for this combination */
 	FString PlatformName;
@@ -60,7 +60,7 @@ public:
 	/**
 	 * Queries whether a configuration is valid for any available platform
 	 */
-	bool IsValidConfiguration(const EBuildConfigurations::Type Configuration, EProjectType ProjectType = EProjectType::Any) const;
+	bool IsValidConfiguration(const EBuildConfiguration Configuration, EProjectType ProjectType = EProjectType::Any) const;
 
 	/**
 	 * Queries whether a platform has any valid configurations
@@ -70,7 +70,7 @@ public:
 	/**
 	 * Queries whether a platform and configuration combination is valid
 	 */
-	bool IsValidPlatformAndConfiguration(const EBuildConfigurations::Type Configuration, const FString& PlatformName, EProjectType ProjectType = EProjectType::Any) const;
+	bool IsValidPlatformAndConfiguration(const EBuildConfiguration Configuration, const FString& PlatformName, EProjectType ProjectType = EProjectType::Any) const;
 
 	/**
 	 * Queries whether a platform can be displayed as an option, even if it's not supported for the specified project type
diff --git a/Engine/Source/Developer/TargetPlatform/Public/Interfaces/ITargetDevice.h b/Engine/Source/Developer/TargetPlatform/Public/Interfaces/ITargetDevice.h
--- a/Engine/Source/Developer/TargetPlatform/Public/Interfaces/ITargetDevice.h
+++ b/Engine/Source/Developer/TargetPlatform/Public/Interfaces/ITargetDevice.h
@@ -312,7 +312,7 @@ public:
 	 * @param OutProcessId Will hold the identifier of the created process (can be NULL).
 	 * @return true on success, false otherwise.
 	 */
-	virtual bool Launch( const FString& AppId, EBuildConfigurations::Type BuildConfiguration, EBuildTargetType TargetType, const FString& Params, uint32* OutProcessId ) = 0;
+	virtual bool Launch( const FString& AppId, EBuildConfiguration BuildConfiguration, EBuildTargetType TargetType, const FString& Params, uint32* OutProcessId ) = 0;
 
 	/**
 	 * Powers off the device.
diff --git a/Engine/Source/Developer/Windows/LiveCoding/Private/LiveCodingModule.cpp b/Engine/Source/Developer/Windows/LiveCoding/Private/LiveCodingModule.cpp
--- a/Engine/Source/Developer/Windows/LiveCoding/Private/LiveCodingModule.cpp
+++ b/Engine/Source/Developer/Windows/LiveCoding/Private/LiveCodingModule.cpp
@@ -288,7 +288,7 @@ bool FLiveCodingModule::StartLiveCoding()
 		// Build the command line
 		FString Arguments;
 		Arguments += FString::Printf(TEXT("%s"), FPlatformMisc::GetUBTPlatform());
-		Arguments += FString::Printf(TEXT(" %s"), EBuildConfigurations::ToString(FApp::GetBuildConfiguration()));
+		Arguments += FString::Printf(TEXT(" %s"), LexToString(FApp::GetBuildConfiguration()));
 		Arguments += FString::Printf(TEXT(" -TargetType=%s"), FPlatformMisc::GetUBTTarget());
 		if(SourceProject.Len() > 0)
 		{
diff --git a/Engine/Source/Developer/Windows/WindowsTargetPlatform/Private/LocalPcTargetDevice.h b/Engine/Source/Developer/Windows/WindowsTargetPlatform/Private/LocalPcTargetDevice.h
--- a/Engine/Source/Developer/Windows/WindowsTargetPlatform/Private/LocalPcTargetDevice.h
+++ b/Engine/Source/Developer/Windows/WindowsTargetPlatform/Private/LocalPcTargetDevice.h
@@ -237,7 +237,7 @@ public:
 		return true;
 	}
 
-	virtual bool Launch( const FString& AppId, EBuildConfigurations::Type BuildConfiguration, EBuildTargetType TargetType, const FString& Params, uint32* OutProcessId ) override
+	virtual bool Launch( const FString& AppId, EBuildConfiguration BuildConfiguration, EBuildTargetType TargetType, const FString& Params, uint32* OutProcessId ) override
 	{
 		// build executable path
 		FString PlatformName = WIN64 ? TEXT("Win64") : TEXT("Win32");
@@ -256,9 +256,9 @@ public:
 			ExecutablePath /= TEXT("UE4Editor");
 		}
 
-		if (BuildConfiguration != EBuildConfigurations::Development)
+		if (BuildConfiguration != EBuildConfiguration::Development)
 		{
-			ExecutablePath += FString::Printf(TEXT("-%s-%s"), *PlatformName, EBuildConfigurations::ToString(BuildConfiguration));
+			ExecutablePath += FString::Printf(TEXT("-%s-%s"), *PlatformName, LexToString(BuildConfiguration));
 		}
 
 		ExecutablePath += TEXT(".exe");
diff --git a/Engine/Source/Editor/LevelEditor/Private/LevelEditor.cpp b/Engine/Source/Editor/LevelEditor/Private/LevelEditor.cpp
--- a/Engine/Source/Editor/LevelEditor/Private/LevelEditor.cpp
+++ b/Engine/Source/Editor/LevelEditor/Private/LevelEditor.cpp
@@ -86,8 +86,8 @@ public:
 		FText RightContentText;
 		FText RightContentTooltip;
 
-		const EBuildConfigurations::Type BuildConfig = FApp::GetBuildConfiguration();
-		if (BuildConfig != EBuildConfigurations::Shipping && BuildConfig != EBuildConfigurations::Development && BuildConfig != EBuildConfigurations::Unknown)
+		const EBuildConfiguration BuildConfig = FApp::GetBuildConfiguration();
+		if (BuildConfig != EBuildConfiguration::Shipping && BuildConfig != EBuildConfiguration::Development && BuildConfig != EBuildConfiguration::Unknown)
 		{
 			Args.Add(TEXT("Config"), EBuildConfigurations::ToText(BuildConfig));
 			RightContentText = FText::Format(NSLOCTEXT("UnrealEditor", "TitleBarRightContentAndConfig", "{ProjectNameWatermarkPrefix} {GameName} [{Config}] {Branch} {EngineVersion}"), Args);
diff --git a/Engine/Source/Editor/MainFrame/Private/MainFrameModule.cpp b/Engine/Source/Editor/MainFrame/Private/MainFrameModule.cpp
--- a/Engine/Source/Editor/MainFrame/Private/MainFrameModule.cpp
+++ b/Engine/Source/Editor/MainFrame/Private/MainFrameModule.cpp
@@ -44,9 +44,9 @@ const FText StaticGetApplicationTitle( const bool bIncludeGameName )
 		Args.Add(TEXT("GameName"), FText::FromString( FString( FApp::GetProjectName())));
 		Args.Add(TEXT("AppTitle"), ApplicationTitle);
 
-		const EBuildConfigurations::Type BuildConfig = FApp::GetBuildConfiguration();
+		const EBuildConfiguration BuildConfig = FApp::GetBuildConfiguration();
 
-		if (BuildConfig != EBuildConfigurations::Shipping && BuildConfig != EBuildConfigurations::Development && BuildConfig != EBuildConfigurations::Unknown)
+		if (BuildConfig != EBuildConfiguration::Shipping && BuildConfig != EBuildConfiguration::Development && BuildConfig != EBuildConfiguration::Unknown)
 		{
 			Args.Add( TEXT("Config"), EBuildConfigurations::ToText(BuildConfig));
 
diff --git a/Engine/Source/Editor/UnrealEd/Private/PlayLevel.cpp b/Engine/Source/Editor/UnrealEd/Private/PlayLevel.cpp
--- a/Engine/Source/Editor/UnrealEd/Private/PlayLevel.cpp
+++ b/Engine/Source/Editor/UnrealEd/Private/PlayLevel.cpp
@@ -2020,16 +2020,16 @@ void UEditorEngine::PlayUsingLauncher()
 		switch (PlayInSettings->LaunchConfiguration)
 		{
 		case LaunchConfig_Debug:
-			LauncherProfile->SetBuildConfiguration(EBuildConfigurations::Debug);
+			LauncherProfile->SetBuildConfiguration(EBuildConfiguration::Debug);
 			break;
 		case LaunchConfig_Development:
-			LauncherProfile->SetBuildConfiguration(EBuildConfigurations::Development);
+			LauncherProfile->SetBuildConfiguration(EBuildConfiguration::Development);
 			break;
 		case LaunchConfig_Test:
-			LauncherProfile->SetBuildConfiguration(EBuildConfigurations::Test);
+			LauncherProfile->SetBuildConfiguration(EBuildConfiguration::Test);
 			break;
 		case LaunchConfig_Shipping:
-			LauncherProfile->SetBuildConfiguration(EBuildConfigurations::Shipping);
+			LauncherProfile->SetBuildConfiguration(EBuildConfiguration::Shipping);
 			break;
 		default:
 			// same as the running editor
diff --git a/Engine/Source/Runtime/Core/Private/GenericPlatform/GenericPlatformCrashContext.cpp b/Engine/Source/Runtime/Core/Private/GenericPlatform/GenericPlatformCrashContext.cpp
--- a/Engine/Source/Runtime/Core/Private/GenericPlatform/GenericPlatformCrashContext.cpp
+++ b/Engine/Source/Runtime/Core/Private/GenericPlatform/GenericPlatformCrashContext.cpp
@@ -313,11 +313,11 @@ void FGenericCrashContext::SerializeContentToBuffer() const
 		}
 	}
 	AddCrashProperty( TEXT( "ExecutableName" ), *NCachedCrashContextProperties::ExecutableName );
-	AddCrashProperty( TEXT( "BuildConfiguration" ), EBuildConfigurations::ToString( FApp::GetBuildConfiguration() ) );
+	AddCrashProperty( TEXT( "BuildConfiguration" ), LexToString( FApp::GetBuildConfiguration() ) );
 	AddCrashProperty( TEXT( "GameSessionID" ), *NCachedCrashContextProperties::GameSessionID );
 	
 	// Unique string specifying the symbols to be used by CrashReporter
-	FString Symbols = FString::Printf( TEXT( "%s-%s-%s" ), FApp::GetBuildVersion(), FPlatformMisc::GetUBTPlatform(), EBuildConfigurations::ToString(FApp::GetBuildConfiguration())).Replace( TEXT( "+" ), TEXT( "*" ));
+	FString Symbols = FString::Printf( TEXT( "%s-%s-%s" ), FApp::GetBuildVersion(), FPlatformMisc::GetUBTPlatform(), LexToString(FApp::GetBuildConfiguration())).Replace( TEXT( "+" ), TEXT( "*" ));
 #ifdef UE_BUILD_FLAVOR
 	Symbols = FString::Printf(TEXT( "%s-%s" ), *Symbols, *FString(UE_BUILD_FLAVOR));
 #endif
diff --git a/Engine/Source/Runtime/Core/Private/GenericPlatform/GenericPlatformMisc.cpp b/Engine/Source/Runtime/Core/Private/GenericPlatform/GenericPlatformMisc.cpp
--- a/Engine/Source/Runtime/Core/Private/GenericPlatform/GenericPlatformMisc.cpp
+++ b/Engine/Source/Runtime/Core/Private/GenericPlatform/GenericPlatformMisc.cpp
@@ -53,59 +53,68 @@ CORE_API void (*ClipboardPasteShim)(FString& Dest) = nullptr;
 /* EBuildConfigurations interface
  *****************************************************************************/
 
-namespace EBuildConfigurations
+void LexFromString( EBuildConfiguration& OutConfiguration, const TCHAR* Configuration )
 {
-	EBuildConfigurations::Type FromString( const FString& Configuration )
+	if (FCString::Stricmp(Configuration, TEXT("Debug")) == 0)
 	{
-		if (FCString::Strcmp(*Configuration, TEXT("Debug")) == 0)
-		{
-			return Debug;
-		}
-		else if (FCString::Strcmp(*Configuration, TEXT("DebugGame")) == 0)
-		{
-			return DebugGame;
-		}
-		else if (FCString::Strcmp(*Configuration, TEXT("Development")) == 0)
-		{
-			return Development;
-		}
-		else if (FCString::Strcmp(*Configuration, TEXT("Shipping")) == 0)
-		{
-			return Shipping;
-		}
-		else if(FCString::Strcmp(*Configuration, TEXT("Test")) == 0)
-		{
-			return Test;
-		}
-
-		return Unknown;
+		OutConfiguration = EBuildConfiguration::Debug;
 	}
-
-	const TCHAR* ToString( EBuildConfigurations::Type Configuration )
+	else if (FCString::Stricmp(Configuration, TEXT("DebugGame")) == 0)
 	{
-		switch (Configuration)
-		{
-			case Debug:
-				return TEXT("Debug");
-
-			case DebugGame:
-				return TEXT("DebugGame");
-
-			case Development:
-				return TEXT("Development");
+		OutConfiguration = EBuildConfiguration::DebugGame;
+	}
+	else if (FCString::Stricmp(Configuration, TEXT("Development")) == 0)
+	{
+		OutConfiguration = EBuildConfiguration::Development;
+	}
+	else if (FCString::Stricmp(Configuration, TEXT("Shipping")) == 0)
+	{
+		OutConfiguration = EBuildConfiguration::Shipping;
+	}
+	else if(FCString::Stricmp(Configuration, TEXT("Test")) == 0)
+	{
+		OutConfiguration = EBuildConfiguration::Test;
+	}
+	else
+	{
+		OutConfiguration = EBuildConfiguration::Unknown;
+	}
+}
 
-			case Shipping:
-				return TEXT("Shipping");
+const TCHAR* LexToString( EBuildConfiguration Configuration )
+{
+	switch (Configuration)
+	{
+	case EBuildConfiguration::Debug:
+		return TEXT("Debug");
+	case EBuildConfiguration::DebugGame:
+		return TEXT("DebugGame");
+	case EBuildConfiguration::Development:
+		return TEXT("Development");
+	case EBuildConfiguration::Shipping:
+		return TEXT("Shipping");
+	case EBuildConfiguration::Test:
+		return TEXT("Test");
+	default:
+		return TEXT("Unknown");
+	}
+}
 
-			case Test:
-				return TEXT("Test");
+namespace EBuildConfigurations
+{
+	EBuildConfiguration FromString( const FString& Configuration )
+	{
+		EBuildConfiguration Result;
+		LexFromString(Result, *Configuration);
+		return Result;
+	}
 
-			default:
-				return TEXT("Unknown");
-		}
+	const TCHAR* ToString( EBuildConfiguration Configuration )
+	{
+		return LexToString(Configuration);
 	}
 
-	FText ToText( EBuildConfigurations::Type Configuration )
+	FText ToText( EBuildConfiguration Configuration )
 	{
 		switch (Configuration)
 		{
@@ -134,30 +143,32 @@ namespace EBuildConfigurations
 /* EBuildTargetType functions
  *****************************************************************************/
 
-EBuildTargetType LexFromString(const TCHAR* Type)
+void LexFromString(EBuildTargetType& OutType, const TCHAR* Type)
 {
 	if (FCString::Strcmp(Type, TEXT("Editor")) == 0)
 	{
-		return EBuildTargetType::Editor;
+		OutType = EBuildTargetType::Editor;
 	}
 	else if (FCString::Strcmp(Type, TEXT("Game")) == 0)
 	{
-		return EBuildTargetType::Game;
+		OutType = EBuildTargetType::Game;
 	}
 	else if (FCString::Strcmp(Type, TEXT("Server")) == 0)
 	{
-		return EBuildTargetType::Server;
+		OutType = EBuildTargetType::Server;
 	}
 	else if (FCString::Strcmp(Type, TEXT("Client")) == 0)
 	{
-		return EBuildTargetType::Client;
+		OutType = EBuildTargetType::Client;
 	}
 	else if (FCString::Strcmp(Type, TEXT("Program")) == 0)
 	{
-		return EBuildTargetType::Program;
+		OutType = EBuildTargetType::Program;
+	}
+	else
+	{
+		OutType = EBuildTargetType::Unknown;
 	}
-
-	return EBuildTargetType::Unknown;
 }
 
 const TCHAR* LexToString(EBuildTargetType Type)
@@ -181,7 +192,9 @@ const TCHAR* LexToString(EBuildTargetType Type)
 
 EBuildTargetType EBuildTargets::FromString(const FString& Target)
 {
-	return LexFromString(*Target);
+	EBuildTargetType Type;
+	LexFromString(Type, *Target);
+	return Type;
 }
 
 const TCHAR* EBuildTargets::ToString(EBuildTargetType Target)
diff --git a/Engine/Source/Runtime/Core/Private/GenericPlatform/GenericPlatformProcess.cpp b/Engine/Source/Runtime/Core/Private/GenericPlatform/GenericPlatformProcess.cpp
--- a/Engine/Source/Runtime/Core/Private/GenericPlatform/GenericPlatformProcess.cpp
+++ b/Engine/Source/Runtime/Core/Private/GenericPlatform/GenericPlatformProcess.cpp
@@ -172,7 +172,7 @@ const TCHAR* FGenericPlatformProcess::ExecutableName(bool bRemoveExtension)
 	return NULL;
 }
 
-FString FGenericPlatformProcess::GenerateApplicationPath( const FString& AppName, EBuildConfigurations::Type BuildConfiguration)
+FString FGenericPlatformProcess::GenerateApplicationPath( const FString& AppName, EBuildConfiguration BuildConfiguration)
 {
 	UE_LOG(LogHAL, Fatal, TEXT("FGenericPlatformProcess::GenerateApplicationPath not implemented on this platform"));
 	return FString();
diff --git a/Engine/Source/Runtime/Core/Private/IOS/IOSPlatformMisc.cpp b/Engine/Source/Runtime/Core/Private/IOS/IOSPlatformMisc.cpp
--- a/Engine/Source/Runtime/Core/Private/IOS/IOSPlatformMisc.cpp
+++ b/Engine/Source/Runtime/Core/Private/IOS/IOSPlatformMisc.cpp
@@ -1954,7 +1954,7 @@ void FIOSCrashContext::GenerateEnsureInfo() const
             Arguments = FString::Printf(TEXT("\"%s/\" -Unattended"), *EnsureLogFolder);
         }
         
-        FString ReportClient = FPaths::ConvertRelativePathToFull(FPlatformProcess::GenerateApplicationPath(TEXT("CrashReportClient"), EBuildConfigurations::Development));
+        FString ReportClient = FPaths::ConvertRelativePathToFull(FPlatformProcess::GenerateApplicationPath(TEXT("CrashReportClient"), EBuildConfiguration::Development));
         FPlatformProcess::ExecProcess(*ReportClient, *Arguments, nullptr, nullptr, nullptr);
     }
 #endif
diff --git a/Engine/Source/Runtime/Core/Private/IOS/IOSPlatformProcess.cpp b/Engine/Source/Runtime/Core/Private/IOS/IOSPlatformProcess.cpp
--- a/Engine/Source/Runtime/Core/Private/IOS/IOSPlatformProcess.cpp
+++ b/Engine/Source/Runtime/Core/Private/IOS/IOSPlatformProcess.cpp
@@ -131,15 +131,15 @@ const TCHAR* FIOSPlatformProcess::ExecutableName(bool bRemoveExtension)
 	return Result;
 }
 
-FString FIOSPlatformProcess::GenerateApplicationPath( const FString& AppName, EBuildConfigurations::Type BuildConfiguration)
+FString FIOSPlatformProcess::GenerateApplicationPath( const FString& AppName, EBuildConfiguration BuildConfiguration)
 {
     SCOPED_AUTORELEASE_POOL;
     
     FString PlatformName = TEXT("IOS");
     FString ExecutableName = AppName;
-    if (BuildConfiguration != EBuildConfigurations::Development && BuildConfiguration != EBuildConfigurations::DebugGame)
+    if (BuildConfiguration != EBuildConfiguration::Development && BuildConfiguration != EBuildConfiguration::DebugGame)
     {
-        ExecutableName += FString::Printf(TEXT("-%s-%s"), *PlatformName, EBuildConfigurations::ToString(BuildConfiguration));
+        ExecutableName += FString::Printf(TEXT("-%s-%s"), *PlatformName, LexToString(BuildConfiguration));
     }
     
     NSURL* CurrentBundleURL = [[NSBundle mainBundle] bundleURL];
diff --git a/Engine/Source/Runtime/Core/Private/Internationalization/ICUInternationalization.cpp b/Engine/Source/Runtime/Core/Private/Internationalization/ICUInternationalization.cpp
--- a/Engine/Source/Runtime/Core/Private/Internationalization/ICUInternationalization.cpp
+++ b/Engine/Source/Runtime/Core/Private/Internationalization/ICUInternationalization.cpp
@@ -436,16 +436,16 @@ void FICUInternationalization::ConditionalInitializeAllowedCultures()
 	// Get our current build config string so we can compare it against the config entries
 	FString BuildConfigString;
 	{
-		EBuildConfigurations::Type BuildConfig = FApp::GetBuildConfiguration();
-		if (BuildConfig == EBuildConfigurations::DebugGame)
+		EBuildConfiguration BuildConfig = FApp::GetBuildConfiguration();
+		if (BuildConfig == EBuildConfiguration::DebugGame)
 		{
 			// Treat DebugGame and Debug as the same for loc purposes
-			BuildConfig = EBuildConfigurations::Debug;
+			BuildConfig = EBuildConfiguration::Debug;
 		}
 
-		if (BuildConfig != EBuildConfigurations::Unknown)
+		if (BuildConfig != EBuildConfiguration::Unknown)
 		{
-			BuildConfigString = EBuildConfigurations::ToString(BuildConfig);
+			BuildConfigString = LexToString(BuildConfig);
 		}
 	}
 
diff --git a/Engine/Source/Runtime/Core/Private/Mac/MacPlatformMisc.cpp b/Engine/Source/Runtime/Core/Private/Mac/MacPlatformMisc.cpp
--- a/Engine/Source/Runtime/Core/Private/Mac/MacPlatformMisc.cpp
+++ b/Engine/Source/Runtime/Core/Private/Mac/MacPlatformMisc.cpp
@@ -225,7 +225,7 @@ struct FMacApplicationInfo
 		// Cache & create the crash report folder.
 		FString ReportPath = FPaths::ConvertRelativePathToFull(FString::Printf(TEXT("%s"), *(FPaths::GameAgnosticSavedDir() / TEXT("Crashes"))));
 		FCStringAnsi::Strcpy(CrashReportPath, PATH_MAX+1, TCHAR_TO_UTF8(*ReportPath));
-		FString ReportClient = FPaths::ConvertRelativePathToFull(FPlatformProcess::GenerateApplicationPath(TEXT("CrashReportClient"), EBuildConfigurations::Development));
+		FString ReportClient = FPaths::ConvertRelativePathToFull(FPlatformProcess::GenerateApplicationPath(TEXT("CrashReportClient"), EBuildConfiguration::Development));
 		FCStringAnsi::Strcpy(CrashReportClient, PATH_MAX+1, TCHAR_TO_UTF8(*ReportClient));
 		IFileManager::Get().MakeDirectory(*ReportPath, true);
 		
@@ -2050,7 +2050,7 @@ void FMacCrashContext::GenerateEnsureInfoAndLaunchReporter() const
 			Arguments += TEXT(" -NoAnalytics");
 		}
 
-		FString ReportClient = FPaths::ConvertRelativePathToFull(FPlatformProcess::GenerateApplicationPath(TEXT("CrashReportClient"), EBuildConfigurations::Development));
+		FString ReportClient = FPaths::ConvertRelativePathToFull(FPlatformProcess::GenerateApplicationPath(TEXT("CrashReportClient"), EBuildConfiguration::Development));
 		FPlatformProcess::ExecProcess(*ReportClient, *Arguments, nullptr, nullptr, nullptr);
 	}
 }
diff --git a/Engine/Source/Runtime/Core/Private/Mac/MacPlatformProcess.cpp b/Engine/Source/Runtime/Core/Private/Mac/MacPlatformProcess.cpp
--- a/Engine/Source/Runtime/Core/Private/Mac/MacPlatformProcess.cpp
+++ b/Engine/Source/Runtime/Core/Private/Mac/MacPlatformProcess.cpp
@@ -68,15 +68,15 @@ void FMacPlatformProcess::FreeDllHandle( void* DllHandle )
 	dlclose( DllHandle );
 }
 
-FString FMacPlatformProcess::GenerateApplicationPath( const FString& AppName, EBuildConfigurations::Type BuildConfiguration)
+FString FMacPlatformProcess::GenerateApplicationPath( const FString& AppName, EBuildConfiguration BuildConfiguration)
 {
 	SCOPED_AUTORELEASE_POOL;
 	
 	FString PlatformName = TEXT("Mac");
 	FString ExecutableName = AppName;
-	if (BuildConfiguration != EBuildConfigurations::Development)
+	if (BuildConfiguration != EBuildConfiguration::Development)
 	{
-		ExecutableName += FString::Printf(TEXT("-%s-%s"), *PlatformName, EBuildConfigurations::ToString(BuildConfiguration));
+		ExecutableName += FString::Printf(TEXT("-%s-%s"), *PlatformName, LexToString(BuildConfiguration));
 	}
 	
 	NSURL* CurrentBundleURL = [[NSBundle mainBundle] bundleURL];
diff --git a/Engine/Source/Runtime/Core/Private/Misc/App.cpp b/Engine/Source/Runtime/Core/Private/Misc/App.cpp
--- a/Engine/Source/Runtime/Core/Private/Misc/App.cpp
+++ b/Engine/Source/Runtime/Core/Private/Misc/App.cpp
@@ -68,22 +68,22 @@ FString FApp::GetEpicProductIdentifier()
 	return FString(TEXT(EPIC_PRODUCT_IDENTIFIER));
 }
 
-EBuildConfigurations::Type FApp::GetBuildConfiguration()
+EBuildConfiguration FApp::GetBuildConfiguration()
 {
 #if UE_BUILD_DEBUG
-	return EBuildConfigurations::Debug;
+	return EBuildConfiguration::Debug;
 
 #elif UE_BUILD_DEVELOPMENT
-	return bIsDebugGame ? EBuildConfigurations::DebugGame : EBuildConfigurations::Development;
+	return bIsDebugGame ? EBuildConfiguration::DebugGame : EBuildConfiguration::Development;
 
 #elif UE_BUILD_SHIPPING
-	return EBuildConfigurations::Shipping;
+	return EBuildConfiguration::Shipping;
 
 #elif UE_BUILD_TEST
-	return EBuildConfigurations::Test;
+	return EBuildConfiguration::Test;
 
 #else
-	return EBuildConfigurations::Unknown;
+	return EBuildConfiguration::Unknown;
 #endif
 }
 
@@ -339,7 +339,7 @@ void FApp::PrintStartupLogMessages()
 	UE_LOG(LogInit, Log, TEXT("Compiled with unrecognized C++ compiler"));
 #endif
 
-	UE_LOG(LogInit, Log, TEXT("Build Configuration: %s"), EBuildConfigurations::ToString(FApp::GetBuildConfiguration()));
+	UE_LOG(LogInit, Log, TEXT("Build Configuration: %s"), LexToString(FApp::GetBuildConfiguration()));
 	UE_LOG(LogInit, Log, TEXT("Branch Name: %s"), *FApp::GetBranchName());
 	FString FilteredString = FCommandLine::IsCommandLineLoggingFiltered() ? TEXT("Filtered ") : TEXT("");
 	UE_LOG(LogInit, Log, TEXT("%sCommand Line: %s"), *FilteredString, FCommandLine::GetForLogging());
diff --git a/Engine/Source/Runtime/Core/Private/Modules/BuildVersion.cpp b/Engine/Source/Runtime/Core/Private/Modules/BuildVersion.cpp
--- a/Engine/Source/Runtime/Core/Private/Modules/BuildVersion.cpp
+++ b/Engine/Source/Runtime/Core/Private/Modules/BuildVersion.cpp
@@ -34,7 +34,7 @@ FString FBuildVersion::GetFileNameForCurrentExecutable()
 #endif
 #if UE_BUILD_DEVELOPMENT
 	// For DebugGame builds we want to read the development version file because it's in an engine folder
-	if (FApp::GetBuildConfiguration() == EBuildConfigurations::DebugGame)
+	if (FApp::GetBuildConfiguration() == EBuildConfiguration::DebugGame)
 	{
 		for (int32 Idx = 0; Idx < 2; Idx++)
 		{
diff --git a/Engine/Source/Runtime/Core/Private/Modules/ModuleManager.cpp b/Engine/Source/Runtime/Core/Private/Modules/ModuleManager.cpp
--- a/Engine/Source/Runtime/Core/Private/Modules/ModuleManager.cpp
+++ b/Engine/Source/Runtime/Core/Private/Modules/ModuleManager.cpp
@@ -1095,7 +1095,7 @@ void FModuleManager::MakeUniqueModuleFilename( const FName InModuleName, FString
 
 const TCHAR *FModuleManager::GetUBTConfiguration()
 {
-	return EBuildConfigurations::ToString(FApp::GetBuildConfiguration());
+	return LexToString(FApp::GetBuildConfiguration());
 }
 
 void FModuleManager::StartProcessingNewlyLoadedObjects()
diff --git a/Engine/Source/Runtime/Core/Private/Modules/ModuleManifest.cpp b/Engine/Source/Runtime/Core/Private/Modules/ModuleManifest.cpp
--- a/Engine/Source/Runtime/Core/Private/Modules/ModuleManifest.cpp
+++ b/Engine/Source/Runtime/Core/Private/Modules/ModuleManifest.cpp
@@ -13,7 +13,7 @@ FModuleManifest::FModuleManifest()
 FString FModuleManifest::GetFileName(const FString& DirectoryName, bool bIsGameFolder)
 {
 #if UE_BUILD_DEVELOPMENT
-	return DirectoryName / ((FApp::GetBuildConfiguration() == EBuildConfigurations::DebugGame && bIsGameFolder)? TEXT(UBT_MODULE_MANIFEST_DEBUGGAME) : TEXT(UBT_MODULE_MANIFEST));
+	return DirectoryName / ((FApp::GetBuildConfiguration() == EBuildConfiguration::DebugGame && bIsGameFolder)? TEXT(UBT_MODULE_MANIFEST_DEBUGGAME) : TEXT(UBT_MODULE_MANIFEST));
 #else
 	return DirectoryName / TEXT(UBT_MODULE_MANIFEST);
 #endif
diff --git a/Engine/Source/Runtime/Core/Private/ProfilingDebugging/CsvProfiler.cpp b/Engine/Source/Runtime/Core/Private/ProfilingDebugging/CsvProfiler.cpp
--- a/Engine/Source/Runtime/Core/Private/ProfilingDebugging/CsvProfiler.cpp
+++ b/Engine/Source/Runtime/Core/Private/ProfilingDebugging/CsvProfiler.cpp
@@ -2784,7 +2784,7 @@ void FCsvProfiler::FinalizeCsvFile()
 
 	// Add metadata
 	FString PlatformStr = FString::Printf(TEXT("%s"), ANSI_TO_TCHAR(FPlatformProperties::IniPlatformName()));
-	FString BuildConfigurationStr = EBuildConfigurations::ToString(FApp::GetBuildConfiguration());
+	FString BuildConfigurationStr = LexToString(FApp::GetBuildConfiguration());
 	FString CommandlineStr = FString("\"") + FCommandLine::Get() + FString("\"");
 	// Strip newlines
 	CommandlineStr.ReplaceInline(TEXT("\n"), TEXT(""));
diff --git a/Engine/Source/Runtime/Core/Private/Unix/UnixPlatformProcess.cpp b/Engine/Source/Runtime/Core/Private/Unix/UnixPlatformProcess.cpp
--- a/Engine/Source/Runtime/Core/Private/Unix/UnixPlatformProcess.cpp
+++ b/Engine/Source/Runtime/Core/Private/Unix/UnixPlatformProcess.cpp
@@ -367,14 +367,14 @@ const TCHAR* FUnixPlatformProcess::ExecutableName(bool bRemoveExtension)
 }
 
 
-FString FUnixPlatformProcess::GenerateApplicationPath( const FString& AppName, EBuildConfigurations::Type BuildConfiguration)
+FString FUnixPlatformProcess::GenerateApplicationPath( const FString& AppName, EBuildConfiguration BuildConfiguration)
 {
 	FString PlatformName = FPlatformProcess::GetBinariesSubdirectory();
 	FString ExecutablePath = FPaths::EngineDir() / FString::Printf(TEXT("Binaries/%s/%s"), *PlatformName, *AppName);
 	
-	if (BuildConfiguration != EBuildConfigurations::Development)
+	if (BuildConfiguration != EBuildConfiguration::Development)
 	{
-		ExecutablePath += FString::Printf(TEXT("-%s-%s"), *PlatformName, EBuildConfigurations::ToString(BuildConfiguration));
+		ExecutablePath += FString::Printf(TEXT("-%s-%s"), *PlatformName, LexToString(BuildConfiguration));
 	}
 	return ExecutablePath;
 }
diff --git a/Engine/Source/Runtime/Core/Private/Windows/WindowsPlatformProcess.cpp b/Engine/Source/Runtime/Core/Private/Windows/WindowsPlatformProcess.cpp
--- a/Engine/Source/Runtime/Core/Private/Windows/WindowsPlatformProcess.cpp
+++ b/Engine/Source/Runtime/Core/Private/Windows/WindowsPlatformProcess.cpp
@@ -105,15 +105,15 @@ void FWindowsPlatformProcess::FreeDllHandle( void* DllHandle )
 	::FreeLibrary((HMODULE)DllHandle);
 }
 
-FString FWindowsPlatformProcess::GenerateApplicationPath( const FString& AppName, EBuildConfigurations::Type BuildConfiguration)
+FString FWindowsPlatformProcess::GenerateApplicationPath( const FString& AppName, EBuildConfiguration BuildConfiguration)
 {
 	FString PlatformName = GetBinariesSubdirectory();
 	FString ExecutablePath = FPaths::EngineDir() / FString::Printf(TEXT("Binaries/%s/%s"), *PlatformName, *AppName);
 	FPaths::MakePlatformFilename(ExecutablePath);
 
-	if (BuildConfiguration != EBuildConfigurations::Development)
+	if (BuildConfiguration != EBuildConfiguration::Development)
 	{
-		ExecutablePath += FString::Printf(TEXT("-%s-%s"), *PlatformName, EBuildConfigurations::ToString(BuildConfiguration));
+		ExecutablePath += FString::Printf(TEXT("-%s-%s"), *PlatformName, LexToString(BuildConfiguration));
 	}
 
 	ExecutablePath += TEXT(".exe");
diff --git a/Engine/Source/Runtime/Core/Public/GenericPlatform/GenericPlatformMisc.h b/Engine/Source/Runtime/Core/Public/GenericPlatform/GenericPlatformMisc.h
--- a/Engine/Source/Runtime/Core/Public/GenericPlatform/GenericPlatformMisc.h
+++ b/Engine/Source/Runtime/Core/Public/GenericPlatform/GenericPlatformMisc.h
@@ -26,39 +26,77 @@ class TFunction;
 #define UE_DEBUG_BREAK() ((void)(FPlatformMisc::IsDebuggerPresent() && ([] () { UE_DEBUG_BREAK_IMPL(); } (), 1)))
 #endif
 
+/**
+ * Available build configurations. Mirorred from UnrealTargetConfiguration.
+ */
+enum class EBuildConfiguration : uint8
+{
+	/** Unknown build configuration. */
+	Unknown,
+
+	/** Debug build. */
+	Debug,
+
+	/** DebugGame build. */
+	DebugGame,
+
+	/** Development build. */
+	Development,
+
+	/** Shipping build. */
+	Shipping,
+
+	/** Test build. */
+	Test
+};
+
+/**
+ * Returns the string representation of the specified EBuildConfiguration value.
+ *
+ * @param Configuration The string to get the EBuildConfiguration for.
+ * @return An EBuildConfiguration value.
+ */
+CORE_API void LexFromString( EBuildConfiguration& OutConfiguration, const TCHAR* Configuration );
+
+/**
+ * Returns the string representation of the specified EBuildConfiguration value.
+ *
+ * @param Configuration The value to get the string for.
+ * @return The string representation.
+ */
+CORE_API const TCHAR* LexToString( EBuildConfiguration Configuration );
+
 namespace EBuildConfigurations
 {
-	/**
-	 * Enumerates build configurations.
-	 */
-	enum Type
-	{
-		/** Unknown build configuration. */
-		Unknown,
+	UE_DEPRECATED(4.24, "EBuildConfigurations::Type is deprecated. Use EBuildConfiguration instead.")
+	typedef EBuildConfiguration Type;
 
-		/** Debug build. */
-		Debug,
+	UE_DEPRECATED(4.24, "EBuildConfigurations::Unknown is deprecated. Use EBuildConfiguration::Unknown instead.")
+	static const EBuildConfiguration Unknown = EBuildConfiguration::Unknown;
 
-		/** DebugGame build. */
-		DebugGame,
+	UE_DEPRECATED(4.24, "EBuildConfigurations::Debug is deprecated. Use EBuildConfiguration::Debug instead.")
+	static const EBuildConfiguration Debug = EBuildConfiguration::Debug;
 
-		/** Development build. */
-		Development,
+	UE_DEPRECATED(4.24, "EBuildConfigurations::DebugGame is deprecated. Use EBuildConfiguration::DebugGame instead.")
+	static const EBuildConfiguration DebugGame = EBuildConfiguration::DebugGame;
 
-		/** Shipping build. */
-		Shipping,
+	UE_DEPRECATED(4.24, "EBuildConfigurations::Development is deprecated. Use EBuildConfiguration::Development instead.")
+	static const EBuildConfiguration Development = EBuildConfiguration::Development;
 
-		/** Test build. */
-		Test
-	};
+	UE_DEPRECATED(4.24, "EBuildConfigurations::Test is deprecated. Use EBuildConfiguration::Test instead.")
+	static const EBuildConfiguration Test = EBuildConfiguration::Test;
+
+	UE_DEPRECATED(4.24, "EBuildConfigurations::Shipping is deprecated. Use EBuildConfiguration::Shipping instead.")
+	static const EBuildConfiguration Shipping = EBuildConfiguration::Shipping;
 
 	/**
 	 * Returns the string representation of the specified EBuildConfiguration value.
 	 *
-	 * @param Configuration The string to get the EBuildConfiguration::Type for.
-	 * @return An EBuildConfiguration::Type value.
+	 * @param Configuration The string to get the EBuildConfiguration for.
+	 * @return An EBuildConfiguration value.
 	 */
-	CORE_API EBuildConfigurations::Type FromString( const FString& Configuration );
+	UE_DEPRECATED(4.24, "EBuildConfigurations::FromString() is deprecated. Use LexFromString() instead.")
+	CORE_API EBuildConfiguration FromString( const FString& Configuration );
 
 	/**
 	 * Returns the string representation of the specified EBuildConfiguration value.
@@ -66,7 +104,8 @@ namespace EBuildConfigurations
 	 * @param Configuration The value to get the string for.
 	 * @return The string representation.
 	 */
-	CORE_API const TCHAR* ToString( EBuildConfigurations::Type Configuration );
+	UE_DEPRECATED(4.24, "EBuildConfigurations::ToString() is deprecated. Use LexToString() instead.")
+	CORE_API const TCHAR* ToString( EBuildConfiguration Configuration );
 
 	/**
 	 * Returns the localized text representation of the specified EBuildConfiguration value.
@@ -74,18 +113,13 @@ namespace EBuildConfigurations
 	 * @param Configuration The value to get the text for.
 	 * @return The localized Build configuration text
 	 */
-	CORE_API FText ToText( EBuildConfigurations::Type Configuration );
-}
-
-FORCEINLINE const TCHAR* LexToString(EBuildConfigurations::Type Configuration)
-{
-	return EBuildConfigurations::ToString(Configuration);
+	CORE_API FText ToText( EBuildConfiguration Configuration );
 }
 
 /**
  * Enumerates build target types.
  */
-enum class EBuildTargetType
+enum class EBuildTargetType : uint8
 {
 	/** Unknown build target. */
 	Unknown,
@@ -109,10 +143,11 @@ enum class EBuildTargetType
 /**
  * Returns the string representation of the specified EBuildTarget value.
  *
- * @param Target The value to get the string for.
+ * @param OutType The value to get the string for.
+ * @param Text The text to parse.
  * @return The string representation.
  */
-CORE_API EBuildTargetType LexFromString(const TCHAR* Type);
+CORE_API void LexFromString(EBuildTargetType& OutType, const TCHAR* Text);
 
 /**
  * Returns the string representation of the specified EBuildTargetType value.
diff --git a/Engine/Source/Runtime/Core/Public/GenericPlatform/GenericPlatformProcess.h b/Engine/Source/Runtime/Core/Public/GenericPlatform/GenericPlatformProcess.h
--- a/Engine/Source/Runtime/Core/Public/GenericPlatform/GenericPlatformProcess.h
+++ b/Engine/Source/Runtime/Core/Public/GenericPlatform/GenericPlatformProcess.h
@@ -268,14 +268,14 @@ struct CORE_API FGenericPlatformProcess
 	 * Generates the path to the specified application or game.
 	 *
 	 * The application must reside in the Engine's binaries directory. The returned path is relative to this
-	 * executable's directory.For example, calling this method with "UE4" and EBuildConfigurations::Debug
+	 * executable's directory.For example, calling this method with "UE4" and EBuildConfiguration::Debug
 	 * on Windows 64-bit will generate the path "../Win64/UE4Editor-Win64-Debug.exe"
 	 *
 	 * @param AppName The name of the application or game.
 	 * @param BuildConfiguration The build configuration of the game.
 	 * @return The generated application path.
 	 */
-	static FString GenerateApplicationPath( const FString& AppName, EBuildConfigurations::Type BuildConfiguration);
+	static FString GenerateApplicationPath( const FString& AppName, EBuildConfiguration BuildConfiguration);
 
 	/**
 	 * Return the prefix of dynamic library (e.g. lib)
diff --git a/Engine/Source/Runtime/Core/Public/IOS/IOSPlatformProcess.h b/Engine/Source/Runtime/Core/Public/IOS/IOSPlatformProcess.h
--- a/Engine/Source/Runtime/Core/Public/IOS/IOSPlatformProcess.h
+++ b/Engine/Source/Runtime/Core/Public/IOS/IOSPlatformProcess.h
@@ -39,7 +39,7 @@ struct CORE_API FIOSPlatformProcess : public FGenericPlatformProcess
 	static void SetupRenderThread();
 	static void SetThreadAffinityMask(uint64 AffinityMask);
 	static const TCHAR* ExecutableName(bool bRemoveExtension = true);
-    static FString GenerateApplicationPath( const FString& AppName, EBuildConfigurations::Type BuildConfiguration);
+    static FString GenerateApplicationPath( const FString& AppName, EBuildConfiguration BuildConfiguration);
 
 private:
 	static void SetupThread(int Priority);
diff --git a/Engine/Source/Runtime/Core/Public/Mac/MacPlatformProcess.h b/Engine/Source/Runtime/Core/Public/Mac/MacPlatformProcess.h
--- a/Engine/Source/Runtime/Core/Public/Mac/MacPlatformProcess.h
+++ b/Engine/Source/Runtime/Core/Public/Mac/MacPlatformProcess.h
@@ -113,7 +113,7 @@ struct CORE_API FMacPlatformProcess : public FGenericPlatformProcess
 	static FString GetCurrentWorkingDirectory();
 	static const TCHAR* ExecutablePath();
 	static const TCHAR* ExecutableName(bool bRemoveExtension = true);
-	static FString GenerateApplicationPath( const FString& AppName, EBuildConfigurations::Type BuildConfiguration);
+	static FString GenerateApplicationPath( const FString& AppName, EBuildConfiguration BuildConfiguration);
 	static const TCHAR* GetModuleExtension();
 	static const TCHAR* GetBinariesSubdirectory();
 	static bool CanLaunchURL(const TCHAR* URL);
diff --git a/Engine/Source/Runtime/Core/Public/Misc/App.h b/Engine/Source/Runtime/Core/Public/Misc/App.h
--- a/Engine/Source/Runtime/Core/Public/Misc/App.h
+++ b/Engine/Source/Runtime/Core/Public/Misc/App.h
@@ -39,7 +39,7 @@ public:
 	 *
 	 * @return The build configuration.
 	 */
-	static EBuildConfigurations::Type GetBuildConfiguration();
+	static EBuildConfiguration GetBuildConfiguration();
 
 #if UE_BUILD_DEVELOPMENT
 	/**
diff --git a/Engine/Source/Runtime/Core/Public/Unix/UnixPlatformProcess.h b/Engine/Source/Runtime/Core/Public/Unix/UnixPlatformProcess.h
--- a/Engine/Source/Runtime/Core/Public/Unix/UnixPlatformProcess.h
+++ b/Engine/Source/Runtime/Core/Public/Unix/UnixPlatformProcess.h
@@ -261,7 +261,7 @@ struct CORE_API FUnixPlatformProcess : public FGenericPlatformProcess
 	static const TCHAR* ApplicationSettingsDir();
 	static void SetCurrentWorkingDirectoryToBaseDir();
 	static FString GetCurrentWorkingDirectory();
-	static FString GenerateApplicationPath(const FString& AppName, EBuildConfigurations::Type BuildConfiguration);
+	static FString GenerateApplicationPath(const FString& AppName, EBuildConfiguration BuildConfiguration);
 	static FString GetApplicationName( uint32 ProcessId );
 	static bool SetProcessLimits(EProcessResource::Type Resource, uint64 Limit);
 	static const TCHAR* ExecutablePath();
diff --git a/Engine/Source/Runtime/Core/Public/Windows/WindowsPlatformProcess.h b/Engine/Source/Runtime/Core/Public/Windows/WindowsPlatformProcess.h
--- a/Engine/Source/Runtime/Core/Public/Windows/WindowsPlatformProcess.h
+++ b/Engine/Source/Runtime/Core/Public/Windows/WindowsPlatformProcess.h
@@ -143,7 +143,7 @@ public:
 	static const FString ShaderWorkingDir();
 	static const TCHAR* ExecutablePath();
 	static const TCHAR* ExecutableName(bool bRemoveExtension = true);
-	static FString GenerateApplicationPath( const FString& AppName, EBuildConfigurations::Type BuildConfiguration);
+	static FString GenerateApplicationPath( const FString& AppName, EBuildConfiguration BuildConfiguration);
 	static const TCHAR* GetModuleExtension();
 	static const TCHAR* GetBinariesSubdirectory();
 	static const FString GetModulesDirectory();
diff --git a/Engine/Source/Runtime/Engine/Private/ChartCreation.cpp b/Engine/Source/Runtime/Engine/Private/ChartCreation.cpp
--- a/Engine/Source/Runtime/Engine/Private/ChartCreation.cpp
+++ b/Engine/Source/Runtime/Engine/Private/ChartCreation.cpp
@@ -154,7 +154,7 @@ void FDumpFPSChartToEndpoint::HandleBasicStats()
 {
 	PrintToEndpoint(FString::Printf(TEXT("--- Begin : FPS chart dump for level '%s'"), *MapName));
 
-	PrintToEndpoint(FString::Printf(TEXT("Dumping FPS chart at %s using build %s in config %s built from changelist %i"), *FDateTime::Now().ToString(), FApp::GetBuildVersion(), EBuildConfigurations::ToString(FApp::GetBuildConfiguration()), GetChangeListNumberForPerfTesting()));
+	PrintToEndpoint(FString::Printf(TEXT("Dumping FPS chart at %s using build %s in config %s built from changelist %i"), *FDateTime::Now().ToString(), FApp::GetBuildVersion(), LexToString(FApp::GetBuildConfiguration()), GetChangeListNumberForPerfTesting()));
 
 	PrintToEndpoint(TEXT("Machine info:"));
 	PrintToEndpoint(FString::Printf(TEXT("\tOS: %s %s"), *OSMajor, *OSMinor));
@@ -372,7 +372,7 @@ protected:
 	{
 		// Add non-bucket params
 		ParamArray.Add(FAnalyticsEventAttribute(TEXT("ChangeList"), GetChangeListNumberForPerfTesting()));
-		ParamArray.Add(FAnalyticsEventAttribute(TEXT("BuildType"), EBuildConfigurations::ToString(FApp::GetBuildConfiguration())));
+		ParamArray.Add(FAnalyticsEventAttribute(TEXT("BuildType"), LexToString(FApp::GetBuildConfiguration())));
 		ParamArray.Add(FAnalyticsEventAttribute(TEXT("DateStamp"), FDateTime::Now().ToString()));
 
 		ParamArray.Add(FAnalyticsEventAttribute(TEXT("OS"), FString::Printf(TEXT("%s %s"), *OSMajor, *OSMinor)));
@@ -522,7 +522,7 @@ protected:
 		// Update non- bucket stats.
 		FPSChartRow = FPSChartRow.Replace(TEXT("TOKEN_MAPNAME"), *FString::Printf(TEXT("%s"), *MapName), ESearchCase::CaseSensitive);
 		FPSChartRow = FPSChartRow.Replace(TEXT("TOKEN_CHANGELIST"), *FString::Printf(TEXT("%i"), GetChangeListNumberForPerfTesting()), ESearchCase::CaseSensitive);
-		FPSChartRow = FPSChartRow.Replace(TEXT("TOKEN_BUILDCONFIG"), EBuildConfigurations::ToString(FApp::GetBuildConfiguration()), ESearchCase::CaseSensitive);
+		FPSChartRow = FPSChartRow.Replace(TEXT("TOKEN_BUILDCONFIG"), LexToString(FApp::GetBuildConfiguration()), ESearchCase::CaseSensitive);
 		FPSChartRow = FPSChartRow.Replace(TEXT("TOKEN_DATESTAMP"), *FString::Printf(TEXT("%s"), *FDateTime::Now().ToString()), ESearchCase::CaseSensitive);
 		FPSChartRow = FPSChartRow.Replace(TEXT("TOKEN_DEVICE_PROFILE"), *DeviceProfileName, ESearchCase::CaseSensitive);
 
@@ -549,7 +549,7 @@ protected:
 		FPSChartRow = FPSChartRow.Replace(TEXT("TOKEN_BOUND_GPU_PERCENT"), *FString::Printf(TEXT("%4.2f"), BoundGPUPct), ESearchCase::CaseSensitive);
 
 		// Add non-bucket params
-		// 		ParamArray.Add(FAnalyticsEventAttribute(TEXT("BuildType"), EBuildConfigurations::ToString(FApp::GetBuildConfiguration())));
+		// 		ParamArray.Add(FAnalyticsEventAttribute(TEXT("BuildType"), LexToString(FApp::GetBuildConfiguration())));
 		// 		ParamArray.Add(FAnalyticsEventAttribute(TEXT("Platform"), FString::Printf(TEXT("%s"), ANSI_TO_TCHAR(FPlatformProperties::IniPlatformName()))));
 
 		// Sum up FrameTimes and GameTimes
diff --git a/Engine/Source/Runtime/Engine/Private/EngineSessionManager.cpp b/Engine/Source/Runtime/Engine/Private/EngineSessionManager.cpp
--- a/Engine/Source/Runtime/Engine/Private/EngineSessionManager.cpp
+++ b/Engine/Source/Runtime/Engine/Private/EngineSessionManager.cpp
@@ -943,7 +943,7 @@ void FEngineSessionManager::StartWatchdog(const FString& RunType, const FString&
 		}
 	}
 
-	FString WatchdogPath = FPaths::ConvertRelativePathToFull(FPlatformProcess::GenerateApplicationPath(TEXT("UnrealWatchdog"), EBuildConfigurations::Development));
+	FString WatchdogPath = FPaths::ConvertRelativePathToFull(FPlatformProcess::GenerateApplicationPath(TEXT("UnrealWatchdog"), EBuildConfiguration::Development));
 
 	TArray< FAnalyticsEventAttribute > WatchdogStartedAttributes;
 	WatchdogStartedAttributes.Add(FAnalyticsEventAttribute(TEXT("RunType"), RunType));
diff --git a/Engine/Source/Runtime/Engine/Private/GameEngine.cpp b/Engine/Source/Runtime/Engine/Private/GameEngine.cpp
--- a/Engine/Source/Runtime/Engine/Private/GameEngine.cpp
+++ b/Engine/Source/Runtime/Engine/Private/GameEngine.cpp
@@ -457,7 +457,7 @@ TSharedRef<SWindow> UGameEngine::CreateGameWindow()
 	Args.Add( TEXT("GameName"), FText::FromString( FApp::GetProjectName() ) );
 	Args.Add( TEXT("PlatformArchitecture"), PlatformBits );
 	Args.Add( TEXT("RHIName"), FText::FromName( LegacyShaderPlatformToShaderFormat( GMaxRHIShaderPlatform ) ) );
-	Args.Add( TEXT("BuildConfiguration"), FText::FromString(EBuildConfigurations::ToString(FApp::GetBuildConfiguration()) ) );
+	Args.Add( TEXT("BuildConfiguration"), FText::FromString(LexToString(FApp::GetBuildConfiguration()) ) );
 
 	/************************************************************************/
 	/************************ Add device name to window title****************/
diff --git a/Engine/Source/Runtime/Engine/Private/UnrealEngine.cpp b/Engine/Source/Runtime/Engine/Private/UnrealEngine.cpp
--- a/Engine/Source/Runtime/Engine/Private/UnrealEngine.cpp
+++ b/Engine/Source/Runtime/Engine/Private/UnrealEngine.cpp
@@ -4175,7 +4175,7 @@ bool UEngine::HandleFlushLogCommand( const TCHAR* Cmd, FOutputDevice& Ar )
 bool UEngine::HandleGameVerCommand( const TCHAR* Cmd, FOutputDevice& Ar )
 {
 	FString VersionString = FString::Printf( TEXT( "GameVersion Branch: %s, Configuration: %s, Build: %s, CommandLine: %s" ),
-		*FApp::GetBranchName(), EBuildConfigurations::ToString( FApp::GetBuildConfiguration() ), FApp::GetBuildVersion(), FCommandLine::Get() );
+		*FApp::GetBranchName(), LexToString( FApp::GetBuildConfiguration() ), FApp::GetBuildVersion(), FCommandLine::Get() );
 
 	Ar.Logf( TEXT("%s"), *VersionString );
 	FPlatformApplicationMisc::ClipboardCopy( *VersionString );
@@ -9143,7 +9143,7 @@ void UEngine::LogPerformanceCapture(UWorld* World, const FString& MapName, const
 
 		FAutomationPerformanceSnapshot PerfSnapshot;
 		PerfSnapshot.Changelist = FString::FromInt(ChangeList);
-		PerfSnapshot.BuildConfiguration = EBuildConfigurations::ToString(FApp::GetBuildConfiguration());
+		PerfSnapshot.BuildConfiguration = LexToString(FApp::GetBuildConfiguration());
 		PerfSnapshot.MapName = MapName;
 		PerfSnapshot.MatineeName = MatineeName;
 		PerfSnapshot.AverageFPS = FString::Printf(TEXT("%0.2f"), GAverageFPS);
diff --git a/Engine/Source/Runtime/Projects/Private/ModuleDescriptor.cpp b/Engine/Source/Runtime/Projects/Private/ModuleDescriptor.cpp
--- a/Engine/Source/Runtime/Projects/Private/ModuleDescriptor.cpp
+++ b/Engine/Source/Runtime/Projects/Private/ModuleDescriptor.cpp
@@ -402,7 +402,7 @@ bool FModuleDescriptor::IsCompiledInCurrentConfiguration() const
 	// Cache the string for the current platform
 	static const FString UBTPlatform = FPlatformMisc::GetUBTPlatform();
 	static const FString UBTTarget = FPlatformMisc::GetUBTTarget();
-	static const FString UBTTargetConfiguration = EBuildConfigurations::ToString(FApp::GetBuildConfiguration());
+	static const FString UBTTargetConfiguration = LexToString(FApp::GetBuildConfiguration());
 
 	// Check the platform is whitelisted
 	if(WhitelistPlatforms.Num() > 0 && !WhitelistPlatforms.Contains(UBTPlatform))
diff --git a/Engine/Source/Runtime/Projects/Private/PluginManager.cpp b/Engine/Source/Runtime/Projects/Private/PluginManager.cpp
--- a/Engine/Source/Runtime/Projects/Private/PluginManager.cpp
+++ b/Engine/Source/Runtime/Projects/Private/PluginManager.cpp
@@ -688,7 +688,7 @@ bool FPluginManager::ConfigureEnabledPlugin(const FPluginReferenceDescriptor& Fi
 			const FPluginReferenceDescriptor& Reference = *NewPluginReferences[Idx];
 
 			// Check if the plugin is required for this platform
-			if(!Reference.IsEnabledForPlatform(FPlatformMisc::GetUBTPlatform()) || !Reference.IsEnabledForTargetConfiguration(EBuildConfigurations::ToString(FApp::GetBuildConfiguration())) || !Reference.IsEnabledForTarget(FPlatformMisc::GetUBTTarget()))
+			if(!Reference.IsEnabledForPlatform(FPlatformMisc::GetUBTPlatform()) || !Reference.IsEnabledForTargetConfiguration(LexToString(FApp::GetBuildConfiguration())) || !Reference.IsEnabledForTarget(FPlatformMisc::GetUBTTarget()))
 			{
 				UE_LOG(LogPluginManager, Verbose, TEXT("Ignoring plugin '%s' for platform/configuration"), *Reference.Name);
 				continue;
