Add support for launching TVOSClient from the editor.

diff --git a/Engine/Source/Developer/LauncherServices/Private/Launcher/LauncherWorker.cpp b/Engine/Source/Developer/LauncherServices/Private/Launcher/LauncherWorker.cpp
--- a/Engine/Source/Developer/LauncherServices/Private/Launcher/LauncherWorker.cpp
+++ b/Engine/Source/Developer/LauncherServices/Private/Launcher/LauncherWorker.cpp
@@ -348,6 +348,10 @@ FString FLauncherWorker::CreateUATCommand( const ILauncherProfileRef& InProfile,
 			{
 				Platforms += TEXT("+IOS");
 			}
+			else if (PlatformInfo->TargetPlatformName == FName("TVOSClient"))
+			{
+				Platforms += TEXT("+TVOS");
+			}
 			else if (PlatformInfo->TargetPlatformName == FName("HoloLens"))
 			{
 				Platforms += TEXT("+HoloLens");
