Enable the new optimized loader from edtor

diff --git a/Engine/Source/Developer/LauncherServices/Private/Launcher/LauncherWorker.cpp b/Engine/Source/Developer/LauncherServices/Private/Launcher/LauncherWorker.cpp
--- a/Engine/Source/Developer/LauncherServices/Private/Launcher/LauncherWorker.cpp
+++ b/Engine/Source/Developer/LauncherServices/Private/Launcher/LauncherWorker.cpp
@@ -424,31 +424,38 @@ FString FLauncherWorker::CreateUATCommand( const ILauncherProfileRef& InProfile,
 
 	if (DeviceGroup.IsValid())
 	{
-		const TArray<FString>& Devices = DeviceGroup->GetDeviceIDs();		
+		const TArray<FString>& Devices = DeviceGroup->GetDeviceIDs();
 
-		// for each deployed device...
-		for (int32 DeviceIndex = 0; DeviceIndex < Devices.Num(); ++DeviceIndex)
+		if (Devices.Num() > 0)
 		{
-			const FString& DeviceId = Devices[DeviceIndex];
-			TSharedPtr<ITargetDeviceProxy> DeviceProxy = DeviceProxyManager->FindProxyDeviceForTargetDevice(DeviceId);
-			if (DeviceProxy.IsValid())
+			// for each deployed device...
+			for (int32 DeviceIndex = 0; DeviceIndex < Devices.Num(); ++DeviceIndex)
 			{
-				AddDeviceToLaunchCommand(DeviceId, DeviceProxy, InProfile, DeviceNames, RoleCommands, bVsyncAdded);
-
-				// also add the credentials, if necessary
-				FString DeviceUser = DeviceProxy->GetDeviceUser();
-				if (DeviceUser.Len() > 0)
+				const FString& DeviceId = Devices[DeviceIndex];
+				TSharedPtr<ITargetDeviceProxy> DeviceProxy = DeviceProxyManager->FindProxyDeviceForTargetDevice(DeviceId);
+				if (DeviceProxy.IsValid())
 				{
-					DeviceCommand += FString::Printf(TEXT(" -deviceuser=%s"), *DeviceUser);
-				}
+					AddDeviceToLaunchCommand(DeviceId, DeviceProxy, InProfile, DeviceNames, RoleCommands, bVsyncAdded);
 
-				FString DeviceUserPassword = DeviceProxy->GetDeviceUserPassword();
-				if (DeviceUserPassword.Len() > 0)
-				{
-					DeviceCommand += FString::Printf(TEXT(" -devicepass=%s"), *DeviceUserPassword);
+					// also add the credentials, if necessary
+					FString DeviceUser = DeviceProxy->GetDeviceUser();
+					if (DeviceUser.Len() > 0)
+					{
+						DeviceCommand += FString::Printf(TEXT(" -deviceuser=%s"), *DeviceUser);
+					}
+
+					FString DeviceUserPassword = DeviceProxy->GetDeviceUserPassword();
+					if (DeviceUserPassword.Len() > 0)
+					{
+						DeviceCommand += FString::Printf(TEXT(" -devicepass=%s"), *DeviceUserPassword);
+					}
 				}
 			}
 		}
+		else
+		{
+			RoleCommands = InProfile->GetDefaultLaunchRole()->GetUATCommandLine();
+		}
 	}
 
 	if (DeviceNames.Len() > 0)
@@ -580,6 +587,11 @@ FString FLauncherWorker::CreateUATCommand( const ILauncherProfileRef& InProfile,
 				UATCommand += TEXT(" -pak");
 			}
 
+			if (InProfile->IsUsingIoStore())
+			{
+				UATCommand += TEXT(" -iostore");
+			}
+
 			if ( InProfile->IsCreatingReleaseVersion() )
 			{
 				UATCommand += TEXT(" -createreleaseversion=");
diff --git a/Engine/Source/Developer/LauncherServices/Private/Profiles/LauncherProfile.h b/Engine/Source/Developer/LauncherServices/Private/Profiles/LauncherProfile.h
--- a/Engine/Source/Developer/LauncherServices/Private/Profiles/LauncherProfile.h
+++ b/Engine/Source/Developer/LauncherServices/Private/Profiles/LauncherProfile.h
@@ -2332,6 +2332,16 @@ public:
 		return EditorExe;
 	}
 
+	virtual void SetUseIoStore(bool bInUseIoStore) override
+	{
+		bUseIoStore = bInUseIoStore;
+	}
+
+	virtual bool IsUsingIoStore() const override
+	{
+		return bUseIoStore;
+	}
+
 	//~ End ILauncherProfile Interface
 
 protected:
@@ -2780,6 +2790,9 @@ private:
 	// Additional command line parameters to set for the application when it launches
 	FString AdditionalCommandLineParameters;
 
+	// Use I/O store.
+	bool bUseIoStore;
+
 private:
 
 	// cook in the editor callbacks (not valid for any other cook mode)
diff --git a/Engine/Source/Developer/LauncherServices/Public/ILauncherProfile.h b/Engine/Source/Developer/LauncherServices/Public/ILauncherProfile.h
--- a/Engine/Source/Developer/LauncherServices/Public/ILauncherProfile.h
+++ b/Engine/Source/Developer/LauncherServices/Public/ILauncherProfile.h
@@ -1259,6 +1259,19 @@ public:
 	 */
 	virtual FOnProfileProjectChanged& OnProjectChanged() = 0;
 
+	/**
+	 * Sets whether to use I/O store for optimized loading.
+	 * @param bUseIoStore Whether to use I/O store
+	 */
+	virtual void SetUseIoStore(bool bUseIoStore) = 0;
+
+	/**
+	 * Using I/O store or not.
+	 *
+	 * @return true if using I/O store
+	 */
+	virtual bool IsUsingIoStore() const = 0;
+
 public:
 
 	/** Virtual destructor. */
diff --git a/Engine/Source/Developer/ProjectLauncher/Private/Widgets/Package/SProjectLauncherPackagingSettings.cpp b/Engine/Source/Developer/ProjectLauncher/Private/Widgets/Package/SProjectLauncherPackagingSettings.cpp
--- a/Engine/Source/Developer/ProjectLauncher/Private/Widgets/Package/SProjectLauncherPackagingSettings.cpp
+++ b/Engine/Source/Developer/ProjectLauncher/Private/Widgets/Package/SProjectLauncherPackagingSettings.cpp
@@ -127,6 +127,23 @@ void SProjectLauncherPackagingSettings::Construct(const FArguments& InArgs, cons
 									.Text(LOCTEXT("IncludePrerequisitesCheckBoxText", "Include an installer for prerequisites of packaged games"))
 								]
 							]
+
+						+ SVerticalBox::Slot()
+							.AutoHeight()
+							.Padding(0.0, 4.0, 0.0, 0.0)
+							[
+								SNew(SCheckBox)
+								.IsEnabled(this, &SProjectLauncherPackagingSettings::IsEditable)
+								.IsChecked(this, &SProjectLauncherPackagingSettings::HandleUseIoStoreCheckBoxIsChecked)
+								.OnCheckStateChanged(this, &SProjectLauncherPackagingSettings::HandleUseIoStoreCheckStateChanged)
+								.Padding(FMargin(4.0f, 0.0f))
+								.ToolTipText(LOCTEXT("UseIoStoreCheckBoxTooltip", "Use container files for optimized loading (I/O Store)"))
+								.Content()
+								[
+									SNew(STextBlock)
+									.Text(LOCTEXT("UseIoStoreCheckBoxText", "Use container files for optimized loading (I/O Store)"))
+								]
+							]
 					]
 			]
 	];
@@ -190,6 +207,28 @@ ECheckBoxState SProjectLauncherPackagingSettings::HandleIncludePrerequisitesChec
 	return ECheckBoxState::Unchecked;
 }
 
+void SProjectLauncherPackagingSettings::HandleUseIoStoreCheckStateChanged(ECheckBoxState NewState)
+{
+	ILauncherProfilePtr SelectedProfile = Model->GetSelectedProfile();
+	if (SelectedProfile.IsValid())
+	{
+		SelectedProfile->SetUseIoStore(NewState == ECheckBoxState::Checked);
+	}
+}
+
+ECheckBoxState SProjectLauncherPackagingSettings::HandleUseIoStoreCheckBoxIsChecked() const
+{
+	ILauncherProfilePtr SelectedProfile = Model->GetSelectedProfile();
+	if (SelectedProfile.IsValid())
+	{
+		if (SelectedProfile->IsUsingIoStore())
+		{
+			return ECheckBoxState::Checked;
+		}
+	}
+	return ECheckBoxState::Unchecked;
+}
+
 FText SProjectLauncherPackagingSettings::HandleDirectoryTitleText() const
 {
 	ILauncherProfilePtr SelectedProfile = Model->GetSelectedProfile();
diff --git a/Engine/Source/Developer/ProjectLauncher/Private/Widgets/Package/SProjectLauncherPackagingSettings.h b/Engine/Source/Developer/ProjectLauncher/Private/Widgets/Package/SProjectLauncherPackagingSettings.h
--- a/Engine/Source/Developer/ProjectLauncher/Private/Widgets/Package/SProjectLauncherPackagingSettings.h
+++ b/Engine/Source/Developer/ProjectLauncher/Private/Widgets/Package/SProjectLauncherPackagingSettings.h
@@ -50,6 +50,9 @@ private:
 	void HandleIncludePrerequisitesCheckStateChanged(ECheckBoxState NewState);
 	ECheckBoxState HandleIncludePrerequisitesCheckBoxIsChecked() const;
 
+	void HandleUseIoStoreCheckStateChanged(ECheckBoxState NewState);
+	ECheckBoxState HandleUseIoStoreCheckBoxIsChecked() const;
+
 	// Callback for getting the content text of the 'Directory' label.
 	FText HandleDirectoryTitleText() const;
 	FText HandleDirectoryPathText() const;
diff --git a/Engine/Source/Editor/MainFrame/Private/Frame/MainFrameActions.cpp b/Engine/Source/Editor/MainFrame/Private/Frame/MainFrameActions.cpp
--- a/Engine/Source/Editor/MainFrame/Private/Frame/MainFrameActions.cpp
+++ b/Engine/Source/Editor/MainFrame/Private/Frame/MainFrameActions.cpp
@@ -737,6 +737,11 @@ void FMainFrameActionCallbacks::PackageProject( const FName InPlatformInfoName )
 		OptionalParams += TEXT(" -pak");
 	}
 
+	if (PackagingSettings->bUseIoStore)
+	{
+		OptionalParams += TEXT(" -iostore");
+	}
+
 	if (PackagingSettings->IncludePrerequisites)
 	{
 		OptionalParams += TEXT(" -prereqs");
diff --git a/Engine/Source/Editor/UnrealEd/Classes/Settings/ProjectPackagingSettings.h b/Engine/Source/Editor/UnrealEd/Classes/Settings/ProjectPackagingSettings.h
--- a/Engine/Source/Editor/UnrealEd/Classes/Settings/ProjectPackagingSettings.h
+++ b/Engine/Source/Editor/UnrealEd/Classes/Settings/ProjectPackagingSettings.h
@@ -173,6 +173,10 @@ public:
 	UPROPERTY(config, EditAnywhere, Category=Packaging)
 	bool UsePakFile;
 
+	/** If enabled, all packages will be put into one or more container files. */
+	UPROPERTY(config, EditAnywhere, Category = Packaging)
+	bool bUseIoStore;
+
 	/** 
 	 * If enabled, will generate pak file chunks.  Assets can be assigned to chunks in the editor or via a delegate (See ShooterGameDelegates.cpp). 
 	 * Can be used for streaming installs (PS4 Playgo, XboxOne Streaming Install, etc)
