IoDispatcher: Added validation errors when using I/O store container files w/o Pak files.
Currently this is needed because I/O store container uses the same manifest/response files as Pak files.

diff --git a/Engine/Source/Developer/LauncherServices/Private/Profiles/LauncherProfile.h b/Engine/Source/Developer/LauncherServices/Private/Profiles/LauncherProfile.h
--- a/Engine/Source/Developer/LauncherServices/Private/Profiles/LauncherProfile.h
+++ b/Engine/Source/Developer/LauncherServices/Private/Profiles/LauncherProfile.h
@@ -1868,6 +1868,7 @@ public:
 		EditorExe = LauncherServicesModule.GetExecutableForCommandlets();
 
 		bNotForLicensees = false;
+		bUseIoStore = false;
 
 		Validate();
 	}
@@ -2335,6 +2336,11 @@ public:
 	virtual void SetUseIoStore(bool bInUseIoStore) override
 	{
 		bUseIoStore = bInUseIoStore;
+
+		if (bUseIoStore)
+		{
+			SetDeployWithUnrealPak(true);
+		}
 	}
 
 	virtual bool IsUsingIoStore() const override
@@ -2476,8 +2482,6 @@ protected:
 
 		}
 
-		
-
 		if (CookMode == ELauncherProfileCookModes::OnTheFly)
 		{
 
@@ -2497,6 +2501,11 @@ protected:
 			ValidationErrors.Add(ELauncherProfileValidationErrors::NoArchiveDirectorySpecified);
 		}
 
+		if (bUseIoStore && !DeployWithUnrealPak)
+		{
+			ValidationErrors.Add(ELauncherProfileValidationErrors::IoStoreRequiresPakFiles);
+		}
+
 		ValidatePlatformSDKs();
 		ValidateDeviceStatus();
 	}
diff --git a/Engine/Source/Developer/LauncherServices/Public/ILauncherProfile.h b/Engine/Source/Developer/LauncherServices/Public/ILauncherProfile.h
--- a/Engine/Source/Developer/LauncherServices/Public/ILauncherProfile.h
+++ b/Engine/Source/Developer/LauncherServices/Public/ILauncherProfile.h
@@ -201,6 +201,9 @@ namespace ELauncherProfileValidationErrors
 
 		/** Device is unauthorized or is locked */
 		LaunchDeviceIsUnauthorized,
+ 
+		/** Using I/O store container file(s) requires using UnrealPak */
+		IoStoreRequiresPakFiles,
 
 		Count
 	};
diff --git a/Engine/Source/Developer/ProjectLauncher/Private/Widgets/Shared/SProjectLauncherValidation.h b/Engine/Source/Developer/ProjectLauncher/Private/Widgets/Shared/SProjectLauncherValidation.h
--- a/Engine/Source/Developer/ProjectLauncher/Private/Widgets/Shared/SProjectLauncherValidation.h
+++ b/Engine/Source/Developer/ProjectLauncher/Private/Widgets/Shared/SProjectLauncherValidation.h
@@ -171,6 +171,11 @@ public:
 				MakeValidationMessage(TEXT("Icons.Error"), LOCTEXT("NoArchiveDirectorySpecifiedError", "The archive step requires a valid directory.").ToString(), ELauncherProfileValidationErrors::NoArchiveDirectorySpecified)
 			];
 
+		VertBox->AddSlot().AutoHeight()
+			[
+				MakeValidationMessage(TEXT("Icons.Error"), LOCTEXT("IoStoreRequiresPakFilesError", "UnrealPak must be selected when using I/O store container file(s)").ToString(), ELauncherProfileValidationErrors::IoStoreRequiresPakFiles)
+			];
+
 		check(VertBox->NumSlots() == ELauncherProfileValidationErrors::Count);
 
 		ChildSlot[VertBox.ToSharedRef()];
diff --git a/Engine/Source/Editor/MainFrame/Private/Frame/MainFrameActions.cpp b/Engine/Source/Editor/MainFrame/Private/Frame/MainFrameActions.cpp
--- a/Engine/Source/Editor/MainFrame/Private/Frame/MainFrameActions.cpp
+++ b/Engine/Source/Editor/MainFrame/Private/Frame/MainFrameActions.cpp
@@ -732,14 +732,17 @@ void FMainFrameActionCallbacks::PackageProject( const FName InPlatformInfoName )
 
 	OptionalParams += GetCookingOptionalParams();
 
-	if (PackagingSettings->UsePakFile)
+	if (PackagingSettings->bUseIoStore)
 	{
-		OptionalParams += TEXT(" -pak");
+		OptionalParams += TEXT(" -iostore");
+
+		// Pak file(s) must be used when using container file(s)
+		PackagingSettings->UsePakFile = true;
 	}
 
-	if (PackagingSettings->bUseIoStore)
+	if (PackagingSettings->UsePakFile)
 	{
-		OptionalParams += TEXT(" -iostore");
+		OptionalParams += TEXT(" -pak");
 	}
 
 	if (PackagingSettings->IncludePrerequisites)
