Fix project launcher for platform extensions (UE-87551).

- Use UBT name to prevent issues with "PlatformClient"-type configurations.

diff --git a/Engine/Source/Developer/LauncherServices/Private/Launcher/LauncherWorker.cpp b/Engine/Source/Developer/LauncherServices/Private/Launcher/LauncherWorker.cpp
--- a/Engine/Source/Developer/LauncherServices/Private/Launcher/LauncherWorker.cpp
+++ b/Engine/Source/Developer/LauncherServices/Private/Launcher/LauncherWorker.cpp
@@ -359,7 +359,7 @@ FString FLauncherWorker::CreateUATCommand( const ILauncherProfileRef& InProfile,
 			else
 			{
 				Platforms += TEXT("+");
-				Platforms += PlatformInfo->VanillaPlatformName.ToString();
+				Platforms += PlatformInfo->UBTTargetId.ToString();
 			}
 
 			// Append any extra UAT flags specified for this platform flavor
