Remove code which specifies -NoCompile when building UAT. This is determined automatically now.

diff --git a/Engine/Source/Developer/LauncherServices/Private/Launcher/LauncherWorker.cpp b/Engine/Source/Developer/LauncherServices/Private/Launcher/LauncherWorker.cpp
--- a/Engine/Source/Developer/LauncherServices/Private/Launcher/LauncherWorker.cpp
+++ b/Engine/Source/Developer/LauncherServices/Private/Launcher/LauncherWorker.cpp
@@ -764,17 +764,6 @@ FString FLauncherWorker::CreateUATCommand( const ILauncherProfileRef& InProfile,
 				if (Profile->IsDeployingIncrementally())
 				{
 					UATCommand += " -iterativedeploy";
-					// @todo Lumin hack maybe? Can we always skip the automation script compiling here when going for fast-as-possible?
-					// we will always compile the automation scripts once, then we stop due to time waste
-					static bool bHasCompiledOnce = false;
-					if (bHasCompiledOnce)
-					{
-						UATCommand += " -nocompile";
-					}
-					else
-					{
-						bHasCompiledOnce = true;
-					}
 				}
 			}
 		case ELauncherProfileDeploymentModes::FileServer:
diff --git a/Engine/Source/Editor/MainFrame/Private/Frame/MainFrameActions.cpp b/Engine/Source/Editor/MainFrame/Private/Frame/MainFrameActions.cpp
--- a/Engine/Source/Editor/MainFrame/Private/Frame/MainFrameActions.cpp
+++ b/Engine/Source/Editor/MainFrame/Private/Frame/MainFrameActions.cpp
@@ -356,9 +356,7 @@ const TCHAR* GetUATCompilationFlags()
 {
 	// We never want to compile editor targets when invoking UAT in this context.
 	// If we are installed or don't have a compiler, we must assume we have a precompiled UAT.
-	return (FApp::GetEngineIsPromotedBuild() || FApp::IsEngineInstalled())
-		? TEXT("-nocompile -nocompileeditor")
-		: TEXT("-nocompileeditor");
+	return TEXT("-nocompileeditor");
 }
 
 FString GetCookingOptionalParams()
