Fix "launch on" not building targets automatically in the installed build. There is now a setting for whether to build a target in the project launcher, which defaults to "auto". This setting uses the same logic as packaging and the default "launch" functionality to determine whether a precompiled binary is available.

diff --git a/Engine/Plugins/Editor/MobileLauncherProfileWizard/Source/MobileLauncherProfileWizard/Private/AndroidProfileWizard.cpp b/Engine/Plugins/Editor/MobileLauncherProfileWizard/Source/MobileLauncherProfileWizard/Private/AndroidProfileWizard.cpp
--- a/Engine/Plugins/Editor/MobileLauncherProfileWizard/Source/MobileLauncherProfileWizard/Private/AndroidProfileWizard.cpp
+++ b/Engine/Plugins/Editor/MobileLauncherProfileWizard/Source/MobileLauncherProfileWizard/Private/AndroidProfileWizard.cpp
@@ -35,7 +35,7 @@ static void SetupAndroidAppProfile(ILauncherProfileRef& AppProfile, const FProfi
 
 	AppProfile->SetBuildUAT(true);
 	// App build configuration
-	AppProfile->SetBuildGame(true);
+	AppProfile->SetBuildMode(ELauncherProfileBuildModes::Auto);
 	AppProfile->SetBuildConfiguration(Params.BuildConfiguration);
 	
 	//// Cooking
@@ -85,7 +85,7 @@ static void SetupAndroidDLCProfile(ILauncherProfileRef& DLCProfile, const FProfi
 
 	DLCProfile->SetBuildUAT(true);
 	// App build configuration
-	DLCProfile->SetBuildGame(false);
+	DLCProfile->SetBuildMode(ELauncherProfileBuildModes::DoNotBuild);
 	DLCProfile->SetBuildConfiguration(Params.BuildConfiguration);
 		
 	//// Cooking
diff --git a/Engine/Plugins/Editor/MobileLauncherProfileWizard/Source/MobileLauncherProfileWizard/Private/IOSProfileWizard.cpp b/Engine/Plugins/Editor/MobileLauncherProfileWizard/Source/MobileLauncherProfileWizard/Private/IOSProfileWizard.cpp
--- a/Engine/Plugins/Editor/MobileLauncherProfileWizard/Source/MobileLauncherProfileWizard/Private/IOSProfileWizard.cpp
+++ b/Engine/Plugins/Editor/MobileLauncherProfileWizard/Source/MobileLauncherProfileWizard/Private/IOSProfileWizard.cpp
@@ -35,7 +35,7 @@ static void SetupIOSAppProfile(ILauncherProfileRef& AppProfile, const FProfilePa
 
 	AppProfile->SetBuildUAT(true);
 	// App build configuration
-	AppProfile->SetBuildGame(true);
+	AppProfile->SetBuildMode(ELauncherProfileBuildModes::Auto);
 	AppProfile->SetBuildConfiguration(Params.BuildConfiguration);
 	
 	//// Cooking
@@ -85,7 +85,7 @@ static void SetupIOSDLCProfile(ILauncherProfileRef& DLCProfile, const FProfilePa
 
 	DLCProfile->SetBuildUAT(true);
 	// App build configuration
-	DLCProfile->SetBuildGame(false);
+	DLCProfile->SetBuildMode(ELauncherProfileBuildModes::DoNotBuild);
 	DLCProfile->SetBuildConfiguration(Params.BuildConfiguration);
 	
 	//// Cooking
diff --git a/Engine/Source/Developer/LauncherServices/Private/Launcher/LauncherWorker.cpp b/Engine/Source/Developer/LauncherServices/Private/Launcher/LauncherWorker.cpp
--- a/Engine/Source/Developer/LauncherServices/Private/Launcher/LauncherWorker.cpp
+++ b/Engine/Source/Developer/LauncherServices/Private/Launcher/LauncherWorker.cpp
@@ -516,17 +516,19 @@ FString FLauncherWorker::CreateUATCommand( const ILauncherProfileRef& InProfile,
 		MapList = TEXT(" -map=") + InitialMap;
 	}
 
+	bool bIsBuilding = InProfile->ShouldBuild();
+
 	// Override the Blueprint nativization method for anything other than "cook by the book" mode. Nativized assets
 	// won't get regenerated otherwise, and we don't want UBT to include generated code assets from a previous cook.
 	// Also disable Blueprint nativization if the profile is not configured to also build code. Otherwise nativized
 	// assets generated at cook time will not be linked into the game's executable prior to stage/deployment phases.
-	if (InProfile->GetCookMode() != ELauncherProfileCookModes::ByTheBook || !InProfile->IsBuilding())
+	if (InProfile->GetCookMode() != ELauncherProfileCookModes::ByTheBook || !bIsBuilding)
 	{
 		UATCommand += TEXT(" -ini:Game:[/Script/UnrealEd.ProjectPackagingSettings]:BlueprintNativizationMethod=Disabled");
 	}
 
 	// build
-	if (InProfile->IsBuilding())
+	if (bIsBuilding)
 	{
 		UATCommand += TEXT(" -build");
 
@@ -864,25 +866,12 @@ FString FLauncherWorker::CreateUATCommand( const ILauncherProfileRef& InProfile,
 
 void FLauncherWorker::CreateAndExecuteTasks( const ILauncherProfileRef& InProfile )
 {
-	// check to see if we need to build by default
-	if (!InProfile->HasProjectSpecified())
-	{
-		FString ProjectPath = FPaths::GetPath(InProfile->GetProjectPath());
-		TArray<FString> OutProjectCodeFilenames;
-		IFileManager::Get().FindFilesRecursive(OutProjectCodeFilenames, *(ProjectPath / TEXT("Source")), TEXT("*.h"), true, false, false);
-		IFileManager::Get().FindFilesRecursive(OutProjectCodeFilenames, *(ProjectPath / TEXT("Source")), TEXT("*.cpp"), true, false, false);
-		ISourceCodeAccessModule& SourceCodeAccessModule = FModuleManager::LoadModuleChecked<ISourceCodeAccessModule>("SourceCodeAccess");
-		if (OutProjectCodeFilenames.Num() > 0 && SourceCodeAccessModule.GetAccessor().CanAccessSourceCode())
-		{
-			InProfile->SetBuildGame(true);
-		}
-	}
 	FPlatformProcess::CreatePipe(ReadPipe, WritePipe);
 
 	// create task chains
 	TaskChain = MakeShareable(new FLauncherVerifyProfileTask());
 	TArray<FString> Platforms;
-	if (InProfile->GetCookMode() == ELauncherProfileCookModes::ByTheBook || InProfile->IsBuilding())
+	if (InProfile->GetCookMode() == ELauncherProfileCookModes::ByTheBook || InProfile->ShouldBuild())
 	{
 		Platforms = InProfile->GetCookedPlatforms();
 	}
diff --git a/Engine/Source/Developer/LauncherServices/Private/Profiles/LauncherProfile.h b/Engine/Source/Developer/LauncherServices/Private/Profiles/LauncherProfile.h
--- a/Engine/Source/Developer/LauncherServices/Private/Profiles/LauncherProfile.h
+++ b/Engine/Source/Developer/LauncherServices/Private/Profiles/LauncherProfile.h
@@ -19,6 +19,10 @@
 #include "GameProjectHelper.h"
 #include "Profiles/LauncherProfileLaunchRole.h"
 #include "PlatformInfo.h"
+#include "TargetReceipt.h"
+#include "DesktopPlatformModule.h"
+
+DEFINE_LOG_CATEGORY_STATIC(LogLauncherProfile, Log, All);
 
 class Error;
 
@@ -42,6 +46,7 @@ enum ELauncherVersion
 	LAUNCHERSERVICES_ADDEDMULTILEVELPATCHING = 25,
 	LAUNCHERSERVICES_ADDEDADDITIONALCOMMANDLINE = 26,
 	LAUNCHERSERVICES_ADDEDINCLUDEPREREQUISITES = 27,
+	LAUNCHERSERVICES_ADDEDBUILDMODE = 28,
 
 	//ADD NEW STUFF HERE
 
@@ -56,6 +61,59 @@ enum ESimpleLauncherVersion
 	LAUNCHERSERVICES_SIMPLEFILEFORMATCHANGE = 2,
 };
 
+LAUNCHERSERVICES_API bool HasPromotedTarget(const TCHAR* BaseDir, const TCHAR* TargetName, const TCHAR* Platform, EBuildConfiguration Configuration, const TCHAR* Architecture)
+{
+	// Get the path to the receipt, and check it exists
+	FString ReceiptPath = FTargetReceipt::GetDefaultPath(BaseDir, TargetName, Platform, Configuration, Architecture);
+	if (!FPaths::FileExists(*ReceiptPath))
+	{
+		UE_LOG(LogLauncherProfile, Log, TEXT("Unable to use promoted target - %s does not exist."), *ReceiptPath);
+		return false;
+	}
+
+	// Read the receipt for this target
+	FTargetReceipt Receipt;
+	if (!Receipt.Read(ReceiptPath))
+	{
+		UE_LOG(LogLauncherProfile, Log, TEXT("Unable to use promoted target - cannot read %s"), *ReceiptPath);
+		return false;
+	}
+
+	// Check the receipt is for a promoted build
+	if (!Receipt.Version.IsPromotedBuild)
+	{
+		UE_LOG(LogLauncherProfile, Log, TEXT("Unable to use promoted target - receipt %s is not for a promoted target"), *ReceiptPath);
+		return false;
+	}
+
+	// Make sure it matches the current build info
+	FEngineVersion ReceiptVersion = Receipt.Version.GetEngineVersion();
+	FEngineVersion CurrentVersion = FEngineVersion::Current();
+	if (!ReceiptVersion.ExactMatch(CurrentVersion))
+	{
+		UE_LOG(LogLauncherProfile, Log, TEXT("Unable to use promoted target - receipt version (%s) is not exact match with current engine version (%s)"), *ReceiptVersion.ToString(), *CurrentVersion.ToString());
+		return false;
+	}
+
+	// Print the matching target info
+	UE_LOG(LogLauncherProfile, Log, TEXT("Found promoted target with matching version at %s"), *ReceiptPath);
+	return true;
+}
+
+bool TryGetDefaultTargetName(const FString& ProjectFile, EBuildTargetType TargetType, FString& OutTargetName)
+{
+	const TArray<FTargetInfo>& Targets = FDesktopPlatformModule::Get()->GetTargetsForProject(ProjectFile);
+	for (const FTargetInfo& Target : Targets)
+	{
+		if (Target.Type == TargetType)
+		{
+			OutTargetName = Target.Name;
+			return true;
+		}
+	}
+	return false;
+}
+
 /**
 * Implements a simple profile which controls the desired output of the Launcher for simple
 */
@@ -653,9 +711,66 @@ public:
 		return InvalidPlatform;
 	}
 
-	virtual bool IsBuilding() const override
+	virtual ELauncherProfileBuildModes::Type GetBuildMode() const override
+	{
+		return BuildMode;
+	}
+
+	virtual bool ShouldBuild() override
 	{
-		return BuildGame;
+		bool bBuild = true;
+		if (BuildMode == ELauncherProfileBuildModes::DoNotBuild)
+		{
+			bBuild = false;
+		}
+		else if (BuildMode == ELauncherProfileBuildModes::Auto)
+		{
+			if (FApp::GetEngineIsPromotedBuild())
+			{
+				TArray<FString> TargetPlatformNames = FindPlatforms();
+				for (const FString& TargetPlatformName : TargetPlatformNames)
+				{
+					// Get the target we're building for
+					const ITargetPlatform* TargetPlatform = GetTargetPlatformManager()->FindTargetPlatform(TargetPlatformName);
+					const PlatformInfo::FPlatformInfo& PlatformInfo = TargetPlatform->GetPlatformInfo();
+
+					// Figure out which target we're building
+					FString ReceiptDir;
+					FString TargetName;
+					if (TryGetDefaultTargetName(FPaths::GetProjectFilePath(), PlatformInfo.PlatformType, TargetName))
+					{
+						ReceiptDir = FPaths::GetPath(FPaths::GetProjectFilePath());
+					}
+					else if (TryGetDefaultTargetName(FString(), PlatformInfo.PlatformType, TargetName))
+					{
+						FText Reason;
+						if (TargetPlatform->RequiresTempTarget(false, BuildConfiguration, false, Reason))
+						{
+							UE_LOG(LogLauncherProfile, Log, TEXT("Project requires temp target (%s)"), *Reason.ToString());
+							ReceiptDir = FPaths::GetPath(FPaths::GetProjectFilePath());
+						}
+						else
+						{
+							UE_LOG(LogLauncherProfile, Log, TEXT("Project does not require temp target"));
+							ReceiptDir = FPaths::EngineDir();
+						}
+					}
+					else
+					{
+						UE_LOG(LogLauncherProfile, Log, TEXT("Unable to find any targets for platform %s - forcing build"), *TargetPlatformName);
+						break;
+					}
+
+					// Check if the existing target is valid
+					FString BuildPlatform = PlatformInfo.UBTTargetId.ToString();
+					if (!HasPromotedTarget(*ReceiptDir, *TargetName, *BuildPlatform, BuildConfiguration, nullptr))
+					{
+						break;
+					}
+				}
+			}
+		}
+		return bBuild;
 	}
 
 	virtual bool IsBuildingUAT() const override
@@ -823,6 +938,7 @@ public:
 		}
 
 		// IMPORTANT: make sure to bump LAUNCHERSERVICES_PROFILEVERSION when modifying this!
+		bool BuildGame = false;
 		Archive << Id
 				<< Name
 				<< Description
@@ -929,6 +1045,15 @@ public:
 		{
 			Archive << IncludePrerequisites;
 		}
+
+		if (Version >= LAUNCHERSERVICES_ADDEDBUILDMODE)
+		{
+			Archive << BuildMode;
+		}
+		else if(Archive.IsLoading())
+		{
+			BuildMode = BuildGame ? ELauncherProfileBuildModes::Build : ELauncherProfileBuildModes::DoNotBuild;
+		}
 		
 		DefaultLaunchRole->Serialize(Archive);
 
@@ -1038,7 +1163,7 @@ public:
 		Writer.WriteValue("LaunchMode", LaunchMode);
 		Writer.WriteValue("PackagingMode", PackagingMode);
 		Writer.WriteValue("PackageDir", PackageDir);
-		Writer.WriteValue("BuildGame", BuildGame);
+		Writer.WriteValue("BuildMode", BuildMode);
 		Writer.WriteValue("ForceClose", ForceClose);
 		Writer.WriteValue("Timeout", (int32)Timeout);
 		Writer.WriteValue("Compressed", Compressed);
@@ -1241,7 +1366,7 @@ public:
 		}
 
 		// build
-		Writer.WriteValue("build", IsBuilding());
+		Writer.WriteValue("build", ShouldBuild());
 
 		// cook
 		switch (GetCookMode())
@@ -1512,7 +1637,7 @@ public:
 	TArray<FString> FindPlatforms()
 	{
 		TArray<FString> Platforms;
-		if (GetCookMode() == ELauncherProfileCookModes::ByTheBook || IsBuilding())
+		if (GetCookMode() == ELauncherProfileCookModes::ByTheBook)
 		{
 			Platforms = GetCookedPlatforms();
 		}
@@ -1668,7 +1793,17 @@ public:
 		LaunchMode = (TEnumAsByte<ELauncherProfileLaunchModes::Type>)((int32)Object.GetNumberField("LaunchMode"));
 		PackagingMode = (TEnumAsByte<ELauncherProfilePackagingModes::Type>)((int32)Object.GetNumberField("PackagingMode"));
 		PackageDir = Object.GetStringField("PackageDir");
-		BuildGame = Object.GetBoolField("BuildGame");
+
+		int64 BuildModeValue;
+		if (Object.TryGetNumberField("BuildMode", BuildModeValue))
+		{
+			BuildMode = (TEnumAsByte<ELauncherProfileBuildModes::Type>)(int32)BuildModeValue;
+		}
+		else
+		{
+			BuildMode = Object.GetBoolField("BuildGame") ? ELauncherProfileBuildModes::Build : ELauncherProfileBuildModes::DoNotBuild;
+		}
+
 		ForceClose = Object.GetBoolField("ForceClose");
 		Timeout = (uint32)Object.GetNumberField("Timeout");
 		Compressed = Object.GetBoolField("Compressed");
@@ -1797,7 +1932,7 @@ public:
 		FInternationalization& I18N = FInternationalization::Get();
 
 		// default build settings
-		BuildGame = !FApp::GetEngineIsPromotedBuild() && !FApp::IsEngineInstalled();
+		BuildMode = ELauncherProfileBuildModes::Auto;
 		BuildUAT = !FApp::GetEngineIsPromotedBuild() && !FApp::IsEngineInstalled();
 
 		// default cook settings
@@ -1873,11 +2008,11 @@ public:
 		Validate();
 	}
 
-	virtual void SetBuildGame(bool Build) override
+	virtual void SetBuildMode(ELauncherProfileBuildModes::Type Mode) override
 	{
-		if (BuildGame != Build)
+		if (BuildMode != Mode)
 		{
-			BuildGame = Build;
+			BuildMode = Mode;
 
 			Validate();
 		}
@@ -2632,8 +2767,8 @@ private:
 	// Holds the cooking mode.
 	TEnumAsByte<ELauncherProfileCookModes::Type> CookMode;
 
-	// Holds a flag indicating whether the game should be built
-	bool BuildGame;
+	// Holds the build mode
+	TEnumAsByte<ELauncherProfileBuildModes::Type> BuildMode;
 
 	// Holds a flag indicating whether UAT should be built
 	bool BuildUAT;
diff --git a/Engine/Source/Developer/LauncherServices/Public/ILauncherProfile.h b/Engine/Source/Developer/LauncherServices/Public/ILauncherProfile.h
--- a/Engine/Source/Developer/LauncherServices/Public/ILauncherProfile.h
+++ b/Engine/Source/Developer/LauncherServices/Public/ILauncherProfile.h
@@ -12,6 +12,26 @@ class FJsonObject;
 class ILauncherProfile;
 class ILauncherSimpleProfile;
 
+namespace ELauncherProfileBuildModes
+{
+	/**
+	 * Enumerates modes in which the launcher cooks builds.
+	 */
+	enum Type
+	{
+		/** Build if there is not already an existing pre-built target available. */
+		Auto,
+
+		/** Always build. */
+		Build,
+
+		/** Do not build. */
+		DoNotBuild,
+	};
+}
+
+
+
 namespace ELauncherProfileCookModes
 {
 	/**
@@ -720,13 +740,21 @@ public:
 	virtual bool ShouldStageBaseReleasePaks() const = 0;
 
 	/**
-	 * Checks whether the game should be built.
+	 * Gets the current build mode.
 	 *
-	 * @return true if building the game, false otherwise.
-	 * @see SetBuildGame
+	 * @return The current build mode.
+	 * @see SetBuildMode
 	 */
-	virtual bool IsBuilding() const = 0;
-    
+	virtual ELauncherProfileBuildModes::Type GetBuildMode() const = 0;
+
+	/**
+	 * Determines whether the current profile requires building
+	 *
+	 * @return The current build mode.
+	 * @see SetBuildMode
+	 */
+	virtual bool ShouldBuild() = 0;
+
 	/**
 	 * Checks whether UAT should be built.
 	 *
@@ -965,12 +993,12 @@ public:
 	virtual void RemoveLaunchRole( const ILauncherProfileLaunchRoleRef& Role ) = 0;
 
 	/**
-	 * Sets whether to build the game.
+	 * Sets the current build mode
 	 *
-	 * @param Build Whether the game should be built.
-	 * @see IsBuilding
+	 * @param Mode Whether the game should be built.
+	 * @see GetBuildMode
 	 */
-	virtual void SetBuildGame( bool Build ) = 0;
+	virtual void SetBuildMode(ELauncherProfileBuildModes::Type Mode) = 0;
 
 	/**
 	 * Sets whether to build UAT.
diff --git a/Engine/Source/Developer/ProjectLauncher/Private/Widgets/Build/SProjectLauncherBuildPage.cpp b/Engine/Source/Developer/ProjectLauncher/Private/Widgets/Build/SProjectLauncherBuildPage.cpp
--- a/Engine/Source/Developer/ProjectLauncher/Private/Widgets/Build/SProjectLauncherBuildPage.cpp
+++ b/Engine/Source/Developer/ProjectLauncher/Private/Widgets/Build/SProjectLauncherBuildPage.cpp
@@ -10,6 +10,7 @@
 #include "Widgets/Layout/SBorder.h"
 #include "Widgets/Text/STextBlock.h"
 #include "Widgets/Input/SEditableTextBox.h"
+#include "Widgets/Input/SComboBox.h"
 
 #include "Widgets/Cook/SProjectLauncherCookedPlatforms.h"
 #include "Widgets/Layout/SExpandableArea.h"
@@ -39,6 +40,19 @@ void SProjectLauncherBuildPage::Construct(const FArguments& InArgs, const TShare
 {
 	Model = InModel;
 
+	// create cook modes menu
+	FMenuBuilder BuildModeMenuBuilder(true, NULL);
+	{
+		FUIAction AutoAction(FExecuteAction::CreateSP(this, &SProjectLauncherBuildPage::HandleBuildModeMenuEntryClicked, ELauncherProfileBuildModes::Auto));
+		BuildModeMenuBuilder.AddMenuEntry(LOCTEXT("BuildMode_AutoAction", "Detect Automatically"), LOCTEXT("BuildMode_AutoActionHint", "Detect whether the project needs to be built automatically."), FSlateIcon(), AutoAction);
+
+		FUIAction BuildAction(FExecuteAction::CreateSP(this, &SProjectLauncherBuildPage::HandleBuildModeMenuEntryClicked, ELauncherProfileBuildModes::Build));
+		BuildModeMenuBuilder.AddMenuEntry(LOCTEXT("BuildMode_BuildAction", "Build"), LOCTEXT("BuildMode_BuildActionHint", "Build the target."), FSlateIcon(), BuildAction);
+
+		FUIAction DoNotBuildAction(FExecuteAction::CreateSP(this, &SProjectLauncherBuildPage::HandleBuildModeMenuEntryClicked, ELauncherProfileBuildModes::DoNotBuild));
+		BuildModeMenuBuilder.AddMenuEntry(LOCTEXT("BuildMode_DoNotBuildAction", "Do Not Build"), LOCTEXT("DoNotCookActionHint", "Do not build the target."), FSlateIcon(), DoNotBuildAction);
+	}
+
 	ChildSlot
 	[
 		SNew(SVerticalBox)
@@ -60,10 +74,18 @@ void SProjectLauncherBuildPage::Construct(const FArguments& InArgs, const TShare
 					.AutoWidth()
 					.Padding(8.0, 0.0, 0.0, 0.0)
 					[
-						// build mode check box
-						SNew(SCheckBox)
-							.IsChecked(this, &SProjectLauncherBuildPage::HandleBuildIsChecked)
-							.OnCheckStateChanged(this, &SProjectLauncherBuildPage::HandleBuildCheckedStateChanged)
+						// cooking mode menu
+						SNew(SComboButton)
+						.ButtonContent()
+						[
+							SNew(STextBlock)
+							.Text(this, &SProjectLauncherBuildPage::HandleBuildModeComboButtonContentText)
+						]
+						.ContentPadding(FMargin(6.0, 2.0))
+						.MenuContent()
+						[
+							BuildModeMenuBuilder.MakeWidget()
+						]
 					]
 			]
 
@@ -203,33 +225,46 @@ bool SProjectLauncherBuildPage::GenerateDSYMForProject(const FString& ProjectNam
 /* SProjectLauncherBuildPage callbacks
  *****************************************************************************/
 
-void SProjectLauncherBuildPage::HandleBuildCheckedStateChanged(ECheckBoxState CheckState)
+FText SProjectLauncherBuildPage::HandleBuildModeComboButtonContentText() const
 {
 	ILauncherProfilePtr SelectedProfile = Model->GetSelectedProfile();
 
 	if (SelectedProfile.IsValid())
 	{
-		SelectedProfile->SetBuildGame(CheckState == ECheckBoxState::Checked);
+		ELauncherProfileBuildModes::Type BuildMode = SelectedProfile->GetBuildMode();
+
+		if (BuildMode == ELauncherProfileBuildModes::Auto)
+		{
+			return LOCTEXT("BuildModeComboButton_Auto", "Detect Automatically");
+		}
+
+		if (BuildMode == ELauncherProfileBuildModes::Build)
+		{
+			return LOCTEXT("BuildModeComboButton_Build", "Build");
+		}
+
+		if (BuildMode == ELauncherProfileBuildModes::DoNotBuild)
+		{
+			return LOCTEXT("BuildModeComboButton_DoNotBuild", "Do not build");
+		}
+
+		return LOCTEXT("BuildModeComboButtonDefaultText", "Select...");
 	}
+
+	return FText();
 }
 
 
-ECheckBoxState SProjectLauncherBuildPage::HandleBuildIsChecked() const
+void SProjectLauncherBuildPage::HandleBuildModeMenuEntryClicked(ELauncherProfileBuildModes::Type BuildMode)
 {
 	ILauncherProfilePtr SelectedProfile = Model->GetSelectedProfile();
 
 	if (SelectedProfile.IsValid())
 	{
-		if (SelectedProfile->IsBuilding())
-		{
-			return ECheckBoxState::Checked;
-		}
+		SelectedProfile->SetBuildMode(BuildMode);
 	}
-
-	return ECheckBoxState::Unchecked;
 }
 
-
 void SProjectLauncherBuildPage::HandleProfileManagerProfileSelected(const ILauncherProfilePtr& SelectedProfile, const ILauncherProfilePtr& PreviousProfile)
 {
 	// reload settings
@@ -291,7 +326,7 @@ EVisibility SProjectLauncherBuildPage::ShowBuildConfiguration() const
 {
 	ILauncherProfilePtr SelectedProfile = Model->GetSelectedProfile();
 
-	if (SelectedProfile.IsValid() && SelectedProfile->IsBuilding())
+	if (SelectedProfile.IsValid())
 	{
 		return EVisibility::Visible;
 	}
diff --git a/Engine/Source/Developer/ProjectLauncher/Private/Widgets/Build/SProjectLauncherBuildPage.h b/Engine/Source/Developer/ProjectLauncher/Private/Widgets/Build/SProjectLauncherBuildPage.h
--- a/Engine/Source/Developer/ProjectLauncher/Private/Widgets/Build/SProjectLauncherBuildPage.h
+++ b/Engine/Source/Developer/ProjectLauncher/Private/Widgets/Build/SProjectLauncherBuildPage.h
@@ -8,6 +8,7 @@
 #include "Input/Reply.h"
 #include "Widgets/SCompoundWidget.h"
 #include "Models/ProjectLauncherModel.h"
+#include "ILauncherProfile.h"
 
 class Error;
 enum class ECheckBoxState : uint8;
@@ -51,11 +52,11 @@ protected:
 
 private:
 
-	/** Callback for changing the checked state of a platform menu check box. */
-	void HandleBuildCheckedStateChanged( ECheckBoxState CheckState );
+	// Callback for getting the content text of the 'Build Mode' combo button.
+	FText HandleBuildModeComboButtonContentText() const;
 
-	/** Callback for determining whether a platform menu entry is checked. */
-	ECheckBoxState HandleBuildIsChecked() const;
+	// Callback for clicking an item in the 'Build Mode' menu.
+	void HandleBuildModeMenuEntryClicked(ELauncherProfileBuildModes::Type BuildMode);
 
 	/** Callback for changing the selected profile in the profile manager. */
 	void HandleProfileManagerProfileSelected( const ILauncherProfilePtr& SelectedProfile, const ILauncherProfilePtr& PreviousProfile );
diff --git a/Engine/Source/Editor/MainFrame/MainFrame.Build.cs b/Engine/Source/Editor/MainFrame/MainFrame.Build.cs
--- a/Engine/Source/Editor/MainFrame/MainFrame.Build.cs
+++ b/Engine/Source/Editor/MainFrame/MainFrame.Build.cs
@@ -41,6 +41,7 @@ public class MainFrame : ModuleRules
 				"UndoHistory",
 				"Analytics",
 				"ToolMenus",
+				"LauncherServices"
 			}
 		);
 
diff --git a/Engine/Source/Editor/MainFrame/Private/Frame/MainFrameActions.cpp b/Engine/Source/Editor/MainFrame/Private/Frame/MainFrameActions.cpp
--- a/Engine/Source/Editor/MainFrame/Private/Frame/MainFrameActions.cpp
+++ b/Engine/Source/Editor/MainFrame/Private/Frame/MainFrameActions.cpp
@@ -783,7 +783,7 @@ void FMainFrameActionCallbacks::PackageProject( const FName InPlatformInfoName )
 			// Check if the receipt is for a matching promoted target
 			FString PlatformName = Platform->GetPlatformInfo().UBTTargetId.ToString();
 
-			extern UNREALED_API bool HasPromotedTarget(const TCHAR* BaseDir, const TCHAR* TargetName, const TCHAR* Platform, EBuildConfiguration Configuration, const TCHAR* Architecture);
+			extern LAUNCHERSERVICES_API bool HasPromotedTarget(const TCHAR* BaseDir, const TCHAR* TargetName, const TCHAR* Platform, EBuildConfiguration Configuration, const TCHAR* Architecture);
 			if (HasPromotedTarget(*BaseDir, *TargetName, *PlatformName, ConfigurationInfo.Configuration, nullptr))
 			{
 				bBuild = false;
diff --git a/Engine/Source/Editor/UnrealEd/Private/PlayLevelViaLauncher.cpp b/Engine/Source/Editor/UnrealEd/Private/PlayLevelViaLauncher.cpp
--- a/Engine/Source/Editor/UnrealEd/Private/PlayLevelViaLauncher.cpp
+++ b/Engine/Source/Editor/UnrealEd/Private/PlayLevelViaLauncher.cpp
@@ -52,57 +52,6 @@ static void HandleOutputReceived(const FString& InMessage)
 	}
 }
 
-UNREALED_API bool HasPromotedTarget(const TCHAR* BaseDir, const TCHAR* TargetName, const TCHAR* Platform, EBuildConfiguration Configuration, const TCHAR* Architecture)
-{
-	// Get the path to the receipt, and check it exists
-	FString ReceiptPath = FTargetReceipt::GetDefaultPath(BaseDir, TargetName, Platform, Configuration, Architecture);
-	if (!FPaths::FileExists(*ReceiptPath))
-	{
-		UE_LOG(LogPlayLevel, Log, TEXT("Unable to use promoted target - %s does not exist."), *ReceiptPath);
-		return false;
-	}
-
-	// Read the receipt for this target
-	FTargetReceipt Receipt;
-	if (!Receipt.Read(ReceiptPath))
-	{
-		UE_LOG(LogPlayLevel, Log, TEXT("Unable to use promoted target - cannot read %s"), *ReceiptPath);
-		return false;
-	}
-
-	// Check the receipt is for a promoted build
-	if (!Receipt.Version.IsPromotedBuild)
-	{
-		UE_LOG(LogPlayLevel, Log, TEXT("Unable to use promoted target - receipt %s is not for a promoted target"), *ReceiptPath);
-		return false;
-	}
-
-	// Make sure it matches the current build info
-	FEngineVersion ReceiptVersion = Receipt.Version.GetEngineVersion();
-	FEngineVersion CurrentVersion = FEngineVersion::Current();
-	if (!ReceiptVersion.ExactMatch(CurrentVersion))
-	{
-		UE_LOG(LogPlayLevel, Log, TEXT("Unable to use promoted target - receipt version (%s) is not exact match with current engine version (%s)"), *ReceiptVersion.ToString(), *CurrentVersion.ToString());
-		return false;
-	}
-
-	return true;
-}
-
-bool TryGetDefaultTargetName(const FString& ProjectFile, EBuildTargetType TargetType, FString& OutTargetName)
-{
-	const TArray<FTargetInfo>& Targets = FDesktopPlatformModule::Get()->GetTargetsForProject(ProjectFile);
-	for (const FTargetInfo& Target : Targets)
-	{
-		if (Target.Type == TargetType)
-		{
-			OutTargetName = Target.Name;
-			return true;
-		}
-	}
-	return false;
-}
-
 void UEditorEngine::StartPlayUsingLauncherSession(FRequestPlaySessionParams& InRequestParams)
 {
 	check(InRequestParams.SessionDestination == EPlaySessionDestinationType::Launcher);
@@ -208,64 +157,26 @@ void UEditorEngine::StartPlayUsingLauncherSession(FRequestPlaySessionParams& InR
 	LauncherSessionInfo->bPlayUsingLauncherHasCode = GameProjectModule.Get().ProjectHasCodeFiles();
 
 	// Figure out if we need to build anything
+	ELauncherProfileBuildModes::Type BuildMode;
 	if (EditorPlaySettings->BuildGameBeforeLaunch == EPlayOnBuildMode::PlayOnBuild_Always)
 	{
-		LauncherSessionInfo->bPlayUsingLauncherBuild = true;
+		BuildMode = ELauncherProfileBuildModes::Build;
 	}
 	else if (EditorPlaySettings->BuildGameBeforeLaunch == EPlayOnBuildMode::PlayOnBuild_Never)
 	{
-		LauncherSessionInfo->bPlayUsingLauncherBuild = false;
+		BuildMode = ELauncherProfileBuildModes::DoNotBuild;
 	}
-	else if (EditorPlaySettings->BuildGameBeforeLaunch == EPlayOnBuildMode::PlayOnBuild_Default)
-	{
-		// If this isn't a promoted build, we always have to build
-		bool bBuildTarget = true;
-		if (FApp::GetEngineIsPromotedBuild())
-		{
-			// Figure out which target we're building
-			FString ReceiptDir;
-			FString TargetName;
-			if (TryGetDefaultTargetName(FPaths::GetProjectFilePath(), EBuildTargetType::Game, TargetName))
-			{
-				ReceiptDir = FPaths::GetPath(FPaths::GetProjectFilePath());
-			}
-			else
-			{
-				TargetName = TEXT("UE4Game");
-
-				FText Reason;
-				if (LaunchPlatform->RequiresTempTarget(LauncherSessionInfo->bPlayUsingLauncherHasCode, BuildConfiguration, false, Reason))
-				{
-					UE_LOG(LogPlayLevel, Log, TEXT("Project requires temp target (%s)"), *Reason.ToString());
-					ReceiptDir = FPaths::GetPath(FPaths::GetProjectFilePath());
-				}
-				else
-				{
-					UE_LOG(LogPlayLevel, Log, TEXT("Project does not require temp target"));
-					ReceiptDir = FPaths::EngineDir();
-				}
-			}
-
-			// Get the target platform
-			const ITargetPlatform* TargetPlatform = GetTargetPlatformManager()->FindTargetPlatform(LaunchPlatformName);
-			FString PlatformName = TargetPlatform->GetPlatformInfo().UBTTargetId.ToString();
-
-			// Check if the existing target is valid
-			if (HasPromotedTarget(*ReceiptDir, *TargetName, *PlatformName, BuildConfiguration, nullptr))
-			{
-				bBuildTarget = false;
-			}
-		}
-		LauncherSessionInfo->bPlayUsingLauncherBuild = bBuildTarget;
-	}
-	else if (EditorPlaySettings->BuildGameBeforeLaunch == EPlayOnBuildMode::PlayOnBuild_IfEditorBuiltLocally)
+	else
 	{
-		LauncherSessionInfo->bPlayUsingLauncherBuild = !FApp::GetEngineIsPromotedBuild();
+		BuildMode = ELauncherProfileBuildModes::Auto;
 	}
 
+	// Assume it's building unless disabled
+	LauncherSessionInfo->bPlayUsingLauncherBuild = (BuildMode != ELauncherProfileBuildModes::DoNotBuild);
+
 	// Setup launch profile, keep the setting here to a minimum.
 	ILauncherProfileRef LauncherProfile = LauncherServicesModule.CreateProfile(TEXT("Launch On Device"));
-	LauncherProfile->SetBuildGame(LauncherSessionInfo->bPlayUsingLauncherBuild);
+	LauncherProfile->SetBuildMode(BuildMode);
 	LauncherProfile->SetBuildConfiguration(BuildConfiguration);
 
 	// select the quickest cook mode based on which in editor cook mode is enabled
