Fixed CIS compile and link errors

diff --git a/Engine/Source/Developer/LauncherServices/Private/Profiles/LauncherProfile.h b/Engine/Source/Developer/LauncherServices/Private/Profiles/LauncherProfile.h
--- a/Engine/Source/Developer/LauncherServices/Private/Profiles/LauncherProfile.h
+++ b/Engine/Source/Developer/LauncherServices/Private/Profiles/LauncherProfile.h
@@ -61,7 +61,7 @@ enum ESimpleLauncherVersion
 	LAUNCHERSERVICES_SIMPLEFILEFORMATCHANGE = 2,
 };
 
-LAUNCHERSERVICES_API bool HasPromotedTarget(const TCHAR* BaseDir, const TCHAR* TargetName, const TCHAR* Platform, EBuildConfiguration Configuration, const TCHAR* Architecture)
+LAUNCHERSERVICES_API inline bool HasPromotedTarget(const TCHAR* BaseDir, const TCHAR* TargetName, const TCHAR* Platform, EBuildConfiguration Configuration, const TCHAR* Architecture)
 {
 	// Get the path to the receipt, and check it exists
 	FString ReceiptPath = FTargetReceipt::GetDefaultPath(BaseDir, TargetName, Platform, Configuration, Architecture);
@@ -100,7 +100,7 @@ LAUNCHERSERVICES_API bool HasPromotedTarget(const TCHAR* BaseDir, const TCHAR* T
 	return true;
 }
 
-bool TryGetDefaultTargetName(const FString& ProjectFile, EBuildTargetType TargetType, FString& OutTargetName)
+inline bool TryGetDefaultTargetName(const FString& ProjectFile, EBuildTargetType TargetType, FString& OutTargetName)
 {
 	const TArray<FTargetInfo>& Targets = FDesktopPlatformModule::Get()->GetTargetsForProject(ProjectFile);
 	for (const FTargetInfo& Target : Targets)
diff --git a/Engine/Source/Runtime/Renderer/Private/HairStrands/HairStrandsTransmittance.h b/Engine/Source/Runtime/Renderer/Private/HairStrands/HairStrandsTransmittance.h
--- a/Engine/Source/Runtime/Renderer/Private/HairStrands/HairStrandsTransmittance.h
+++ b/Engine/Source/Runtime/Renderer/Private/HairStrands/HairStrandsTransmittance.h
@@ -10,6 +10,8 @@
 #include "RendererInterface.h"
 #include "RenderGraphResources.h"
 
+class FLightSceneInfo;
+
 struct FHairStrandsTransmittanceMaskData
 {
 	TRefCountPtr<FPooledRDGBuffer>	TransmittanceMask;
