Fixed Mac linker error

#_rb
#rnx

diff --git a/Engine/Source/Developer/LauncherServices/Private/Profiles/LauncherProfile.h b/Engine/Source/Developer/LauncherServices/Private/Profiles/LauncherProfile.h
--- a/Engine/Source/Developer/LauncherServices/Private/Profiles/LauncherProfile.h
+++ b/Engine/Source/Developer/LauncherServices/Private/Profiles/LauncherProfile.h
@@ -61,44 +61,7 @@ enum ESimpleLauncherVersion
 	LAUNCHERSERVICES_SIMPLEFILEFORMATCHANGE = 2,
 };
 
-LAUNCHERSERVICES_API inline bool HasPromotedTarget(const TCHAR* BaseDir, const TCHAR* TargetName, const TCHAR* Platform, EBuildConfiguration Configuration, const TCHAR* Architecture)
-{
-	// Get the path to the receipt, and check it exists
-	FString ReceiptPath = FTargetReceipt::GetDefaultPath(BaseDir, TargetName, Platform, Configuration, Architecture);
-	if (!FPaths::FileExists(*ReceiptPath))
-	{
-		UE_LOG(LogLauncherProfile, Log, TEXT("Unable to use promoted target - %s does not exist."), *ReceiptPath);
-		return false;
-	}
-
-	// Read the receipt for this target
-	FTargetReceipt Receipt;
-	if (!Receipt.Read(ReceiptPath))
-	{
-		UE_LOG(LogLauncherProfile, Log, TEXT("Unable to use promoted target - cannot read %s"), *ReceiptPath);
-		return false;
-	}
-
-	// Check the receipt is for a promoted build
-	if (!Receipt.Version.IsPromotedBuild)
-	{
-		UE_LOG(LogLauncherProfile, Log, TEXT("Unable to use promoted target - receipt %s is not for a promoted target"), *ReceiptPath);
-		return false;
-	}
-
-	// Make sure it matches the current build info
-	FEngineVersion ReceiptVersion = Receipt.Version.GetEngineVersion();
-	FEngineVersion CurrentVersion = FEngineVersion::Current();
-	if (!ReceiptVersion.ExactMatch(CurrentVersion))
-	{
-		UE_LOG(LogLauncherProfile, Log, TEXT("Unable to use promoted target - receipt version (%s) is not exact match with current engine version (%s)"), *ReceiptVersion.ToString(), *CurrentVersion.ToString());
-		return false;
-	}
-
-	// Print the matching target info
-	UE_LOG(LogLauncherProfile, Log, TEXT("Found promoted target with matching version at %s"), *ReceiptPath);
-	return true;
-}
+LAUNCHERSERVICES_API bool HasPromotedTarget(const TCHAR* BaseDir, const TCHAR* TargetName, const TCHAR* Platform, EBuildConfiguration Configuration, const TCHAR* Architecture);
 
 inline bool TryGetDefaultTargetName(const FString& ProjectFile, EBuildTargetType TargetType, FString& OutTargetName)
 {
diff --git a/Engine/Source/Developer/LauncherServices/Private/Profiles/LauncherProfileManager.cpp b/Engine/Source/Developer/LauncherServices/Private/Profiles/LauncherProfileManager.cpp
--- a/Engine/Source/Developer/LauncherServices/Private/Profiles/LauncherProfileManager.cpp
+++ b/Engine/Source/Developer/LauncherServices/Private/Profiles/LauncherProfileManager.cpp
@@ -10,6 +10,44 @@
 #include "Profiles/LauncherDeviceGroup.h"
 #include "Profiles/LauncherProfile.h"
 
+LAUNCHERSERVICES_API bool HasPromotedTarget(const TCHAR* BaseDir, const TCHAR* TargetName, const TCHAR* Platform, EBuildConfiguration Configuration, const TCHAR* Architecture)
+{
+	// Get the path to the receipt, and check it exists
+	FString ReceiptPath = FTargetReceipt::GetDefaultPath(BaseDir, TargetName, Platform, Configuration, Architecture);
+	if (!FPaths::FileExists(*ReceiptPath))
+	{
+		UE_LOG(LogLauncherProfile, Log, TEXT("Unable to use promoted target - %s does not exist."), *ReceiptPath);
+		return false;
+	}
+
+	// Read the receipt for this target
+	FTargetReceipt Receipt;
+	if (!Receipt.Read(ReceiptPath))
+	{
+		UE_LOG(LogLauncherProfile, Log, TEXT("Unable to use promoted target - cannot read %s"), *ReceiptPath);
+		return false;
+	}
+
+	// Check the receipt is for a promoted build
+	if (!Receipt.Version.IsPromotedBuild)
+	{
+		UE_LOG(LogLauncherProfile, Log, TEXT("Unable to use promoted target - receipt %s is not for a promoted target"), *ReceiptPath);
+		return false;
+	}
+
+	// Make sure it matches the current build info
+	FEngineVersion ReceiptVersion = Receipt.Version.GetEngineVersion();
+	FEngineVersion CurrentVersion = FEngineVersion::Current();
+	if (!ReceiptVersion.ExactMatch(CurrentVersion))
+	{
+		UE_LOG(LogLauncherProfile, Log, TEXT("Unable to use promoted target - receipt version (%s) is not exact match with current engine version (%s)"), *ReceiptVersion.ToString(), *CurrentVersion.ToString());
+		return false;
+	}
+
+	// Print the matching target info
+	UE_LOG(LogLauncherProfile, Log, TEXT("Found promoted target with matching version at %s"), *ReceiptPath);
+	return true;
+}
 
 /* ILauncherProfileManager structors
  *****************************************************************************/
