- Fixed ProjectLauncher/UnrealFrontend/LaunchOn to respect the Target32Bit windows project setting

diff --git a/Engine/Source/Developer/LauncherServices/Private/Launcher/LauncherWorker.cpp b/Engine/Source/Developer/LauncherServices/Private/Launcher/LauncherWorker.cpp
--- a/Engine/Source/Developer/LauncherServices/Private/Launcher/LauncherWorker.cpp
+++ b/Engine/Source/Developer/LauncherServices/Private/Launcher/LauncherWorker.cpp
@@ -336,7 +336,19 @@ FString FLauncherWorker::CreateUATCommand( const ILauncherProfileRef& InProfile,
 			{
 				Platforms += TEXT("+LinuxAArch64");
 			}
-			else if (PlatformInfo->TargetPlatformName == FName("WindowsNoEditor") || PlatformInfo->TargetPlatformName == FName("Windows") || PlatformInfo->TargetPlatformName == FName("WindowsClient"))
+			else if (PlatformInfo->TargetPlatformName == FName("WindowsNoEditor"))
+			{
+				// find out if the project is targeting 32bit
+				FConfigFile ProjectEngineConfig;
+				FString ProjectDir = FPaths::GetPath(InProfile->GetProjectPath());
+				FConfigCacheIni::LoadExternalIniFile(ProjectEngineConfig, TEXT("Engine"), *FPaths::EngineConfigDir(), *FPaths::Combine(ProjectDir, TEXT("Config/")), true, TEXT("Windows"));
+				bool bTarget32Bit = false;
+				ProjectEngineConfig.GetBool(TEXT("/Script/WindowsTargetPlatform.WindowsTargetSettings"), TEXT("bTarget32Bit"), bTarget32Bit);
+
+				// if target wants 32-bit, use 32-bit
+				Platforms += bTarget32Bit ? TEXT("+Win32") : TEXT("+Win64");
+			}
+			else if (PlatformInfo->TargetPlatformName == FName("Windows") || PlatformInfo->TargetPlatformName == FName("WindowsClient"))
 			{
 				Platforms += TEXT("+Win64");
 			}
diff --git a/Engine/Source/Developer/Windows/WindowsTargetPlatform/Classes/WindowsTargetSettings.h b/Engine/Source/Developer/Windows/WindowsTargetPlatform/Classes/WindowsTargetSettings.h
--- a/Engine/Source/Developer/Windows/WindowsTargetPlatform/Classes/WindowsTargetSettings.h
+++ b/Engine/Source/Developer/Windows/WindowsTargetPlatform/Classes/WindowsTargetSettings.h
@@ -76,7 +76,7 @@ public:
 	/**
 	 * Determines if data is cooked for 32bit (or 64bit) exe
 	 */
-	UPROPERTY(EditAnywhere, config, Category = "OS Info", Meta = (DisplayName = "Target 32bit"))
+	UPROPERTY(EditAnywhere, config, Category = "OS Info", Meta = (DisplayName = "Target 32bit (should restart editor if changed)"))
 	bool bTarget32Bit;
 
 	/** Sample rate to run the audio mixer with. */
