- Fix CIS error

diff --git a/Engine/Source/Developer/LauncherServices/Private/Launcher/LauncherWorker.cpp b/Engine/Source/Developer/LauncherServices/Private/Launcher/LauncherWorker.cpp
--- a/Engine/Source/Developer/LauncherServices/Private/Launcher/LauncherWorker.cpp
+++ b/Engine/Source/Developer/LauncherServices/Private/Launcher/LauncherWorker.cpp
@@ -17,6 +17,7 @@
 #include "Launcher/LauncherUATTask.h"
 #include "Launcher/LauncherVerifyProfileTask.h"
 #include "PlatformInfo.h"
+#include "Misc/ConfigCacheIni.h"
 
 
 #define LOCTEXT_NAMESPACE "LauncherWorker"
