- 32bit windows project setting now should work in all modes

diff --git a/Engine/Source/Developer/LauncherServices/Private/Launcher/LauncherWorker.cpp b/Engine/Source/Developer/LauncherServices/Private/Launcher/LauncherWorker.cpp
--- a/Engine/Source/Developer/LauncherServices/Private/Launcher/LauncherWorker.cpp
+++ b/Engine/Source/Developer/LauncherServices/Private/Launcher/LauncherWorker.cpp
@@ -346,8 +346,16 @@ FString FLauncherWorker::CreateUATCommand( const ILauncherProfileRef& InProfile,
 				bool bTarget32Bit = false;
 				ProjectEngineConfig.GetBool(TEXT("/Script/WindowsTargetPlatform.WindowsTargetSettings"), TEXT("bTarget32Bit"), bTarget32Bit);
 
+				// normally this would get the 64bit one, so if we need the 32-bit one, swap it out
+				if (bTarget32Bit)
+				{
+					PlatformInfo = PlatformInfo::FindPlatformInfo(TEXT("WindowsNoEditorWin32"));
+					check(PlatformInfo != nullptr);
+				}
+			
 				// if target wants 32-bit, use 32-bit
-				Platforms += bTarget32Bit ? TEXT("+Win32") : TEXT("+Win64");
+				Platforms += TEXT("+");
+				Platforms += PlatformInfo->UBTTargetId.ToString();
 			}
 			else if (PlatformInfo->TargetPlatformName == FName("Windows") || PlatformInfo->TargetPlatformName == FName("WindowsClient"))
 			{
diff --git a/Engine/Source/Developer/Windows/WindowsTargetPlatform/Classes/WindowsTargetSettings.h b/Engine/Source/Developer/Windows/WindowsTargetPlatform/Classes/WindowsTargetSettings.h
--- a/Engine/Source/Developer/Windows/WindowsTargetPlatform/Classes/WindowsTargetSettings.h
+++ b/Engine/Source/Developer/Windows/WindowsTargetPlatform/Classes/WindowsTargetSettings.h
@@ -76,7 +76,7 @@ public:
 	/**
 	 * Determines if data is cooked for 32bit (or 64bit) exe
 	 */
-	UPROPERTY(EditAnywhere, config, Category = "OS Info", Meta = (DisplayName = "Target 32bit (should restart editor if changed)"))
+	UPROPERTY(EditAnywhere, config, Category = "OS Info", Meta = (DisplayName = "Target 32bit", ConfigRestartRequired=true))
 	bool bTarget32Bit;
 
 	/** Sample rate to run the audio mixer with. */
