Another dotnet hang fix. Passing ReadPipe into CreateProcess fixes for now but not 100% correct fix

diff --git a/Engine/Source/Developer/LauncherServices/Private/Launcher/LauncherUATTask.h b/Engine/Source/Developer/LauncherServices/Private/Launcher/LauncherUATTask.h
--- a/Engine/Source/Developer/LauncherServices/Private/Launcher/LauncherUATTask.h
+++ b/Engine/Source/Developer/LauncherServices/Private/Launcher/LauncherUATTask.h
@@ -81,7 +81,7 @@ protected:
 		UATCommandLine += CommandLine;
 
 		// launch UAT and monitor its progress
-		ProcessHandle = FPlatformProcess::CreateProc(*(ExecutablePath / Executable), *UATCommandLine, false, true, true, NULL, 0, *ExecutablePath, WritePipe);
+		ProcessHandle = FPlatformProcess::CreateProc(*(ExecutablePath / Executable), *UATCommandLine, false, true, true, NULL, 0, *ExecutablePath, WritePipe, ReadPipe);
 
 		while (FPlatformProcess::IsProcRunning(ProcessHandle) && !EndTextFound)
 		{
