fix uninitialized variable that was causing Project Launcher to display turnkey GUI when launching Windows builds.

#rnx

diff --git a/Engine/Source/Developer/LauncherServices/Private/Profiles/LauncherProfile.h b/Engine/Source/Developer/LauncherServices/Private/Profiles/LauncherProfile.h
--- a/Engine/Source/Developer/LauncherServices/Private/Profiles/LauncherProfile.h
+++ b/Engine/Source/Developer/LauncherServices/Private/Profiles/LauncherProfile.h
@@ -1932,6 +1932,7 @@ public:
 
 		bNotForLicensees = false;
 		bUseIoStore = false;
+		bShouldUpdateFlash = false;
 		bMakeBinaryConfig = false;
 
 		Validate();
