Pass the culture list override from the project launcher for cook-by-the-book cooks

This override is only used if the list isn't empty, otherwise the project default is used. Cook-on-the-fly isn't impacted by this change, as it cooks assets on-demand so doesn't need to pre-cook a certain list of cultures.

diff --git a/Engine/Source/Developer/LauncherServices/Private/Launcher/LauncherWorker.cpp b/Engine/Source/Developer/LauncherServices/Private/Launcher/LauncherWorker.cpp
--- a/Engine/Source/Developer/LauncherServices/Private/Launcher/LauncherWorker.cpp
+++ b/Engine/Source/Developer/LauncherServices/Private/Launcher/LauncherWorker.cpp
@@ -526,6 +526,17 @@ FString FLauncherWorker::CreateUATCommand( const ILauncherProfileRef& InProfile,
 		MapList = TEXT(" -map=") + InitialMap;
 	}
 
+	// culture list
+	FString CultureList;
+	{
+		const TArray<FString>& CookedCultures = InProfile->GetCookedCultures();
+		if (CookedCultures.Num() > 0 && (InProfile->GetCookMode() == ELauncherProfileCookModes::ByTheBook || InProfile->GetCookMode() == ELauncherProfileCookModes::ByTheBookInEditor))
+		{
+			CultureList += TEXT(" -CookCultures=");
+			CultureList += FString::Join(CookedCultures, TEXT("+"));
+		}
+	}
+
 	bool bIsBuilding = InProfile->ShouldBuild();
 
 	// build
@@ -551,6 +562,7 @@ FString FLauncherWorker::CreateUATCommand( const ILauncherProfileRef& InProfile,
 			UATCommand += TEXT(" -cook");
 
 			UATCommand += MapList;
+			UATCommand += CultureList;
 
 			if (InProfile->IsCookingUnversioned())
 			{
@@ -709,6 +721,7 @@ FString FLauncherWorker::CreateUATCommand( const ILauncherProfileRef& InProfile,
 		break;
 	case ELauncherProfileCookModes::ByTheBookInEditor:
 		UATCommand += MapList;
+		UATCommand += CultureList;
 		UATCommand += TEXT(" -skipcook -CookInEditor"); // don't cook anything the editor is doing it ;)
 		if (InProfile->IsPackingWithUnrealPak())
 		{
