Fixed "use iostore" and "make binary config" settings not being saved to the launcher profile

diff --git a/Engine/Source/Developer/LauncherServices/Private/Profiles/LauncherProfile.h b/Engine/Source/Developer/LauncherServices/Private/Profiles/LauncherProfile.h
--- a/Engine/Source/Developer/LauncherServices/Private/Profiles/LauncherProfile.h
+++ b/Engine/Source/Developer/LauncherServices/Private/Profiles/LauncherProfile.h
@@ -47,6 +47,8 @@ enum ELauncherVersion
 	LAUNCHERSERVICES_ADDEDADDITIONALCOMMANDLINE = 26,
 	LAUNCHERSERVICES_ADDEDINCLUDEPREREQUISITES = 27,
 	LAUNCHERSERVICES_ADDEDBUILDMODE = 28,
+	LAUNCHERSERVICES_ADDEDUSEIOSTORE = 29,
+	LAUNCHERSERVICES_ADDEDMAKEBINARYCONFIG = 30,
 
 	//ADD NEW STUFF HERE
 
@@ -1013,6 +1015,17 @@ public:
 		{
 			Archive << BuildMode;
 		}
+
+		if (Version >= LAUNCHERSERVICES_ADDEDUSEIOSTORE)
+		{
+			Archive << bUseIoStore;
+		}
+
+		if (Version >= LAUNCHERSERVICES_ADDEDMAKEBINARYCONFIG)
+		{
+			Archive << bMakeBinaryConfig;
+		}
+
 		else if(Archive.IsLoading())
 		{
 			BuildMode = BuildGame ? ELauncherProfileBuildModes::Build : ELauncherProfileBuildModes::DoNotBuild;
@@ -1153,6 +1166,8 @@ public:
 		Writer.WriteValue("ArchiveDirectory", ArchiveDir);
 		Writer.WriteValue("AdditionalCommandLineParameters", AdditionalCommandLineParameters);
 		Writer.WriteValue("IncludePrerequisites", IncludePrerequisites);
+		Writer.WriteValue("UseIoStore", bUseIoStore);
+		Writer.WriteValue("MakeBinaryConfig", bMakeBinaryConfig);
 
 		// serialize the default launch role
 		DefaultLaunchRole->Save(Writer, TEXT("DefaultRole"));
@@ -1800,6 +1815,16 @@ public:
 			IncludePrerequisites = Object.GetBoolField("IncludePrerequisites");
 		}
 
+		if (Version >= LAUNCHERSERVICES_ADDEDUSEIOSTORE)
+		{
+			bUseIoStore = Object.GetBoolField("UseIoStore");
+		}
+
+		if (Version >= LAUNCHERSERVICES_ADDEDMAKEBINARYCONFIG)
+		{
+			bMakeBinaryConfig = Object.GetBoolField("MakeBinaryConfig");
+		}
+
 		// load the default launch role
 		TSharedPtr<FJsonObject> Role = Object.GetObjectField("DefaultRole");
 		DefaultLaunchRole->Load(*(Role.Get()));
