LauncherServices: Remove references to UE4

#trivial

diff --git a/Engine/Source/Developer/LauncherServices/Private/LauncherServicesModule.cpp b/Engine/Source/Developer/LauncherServices/Private/LauncherServicesModule.cpp
--- a/Engine/Source/Developer/LauncherServices/Private/LauncherServicesModule.cpp
+++ b/Engine/Source/Developer/LauncherServices/Private/LauncherServicesModule.cpp
@@ -72,7 +72,7 @@ public:
 #if WITH_EDITOR
 		ExecutableName = FString(FPlatformProcess::ExecutablePath());
 #if PLATFORM_WINDOWS
-		// turn UE4editor into UE4editor-cmd
+		// turn UnreadEditor into UnrealEditor-cmd
 		if (ExecutableName.EndsWith(".exe", ESearchCase::IgnoreCase) && !FPaths::GetBaseFilename(ExecutableName).EndsWith("-cmd", ESearchCase::IgnoreCase))
 		{
 			FString NewExeName = ExecutableName.Left(ExecutableName.Len() - 4) + "-Cmd.exe";
@@ -82,7 +82,7 @@ public:
 			}
 		}
 #elif PLATFORM_MAC
-		// turn UE4editor into UE4editor-cmd
+		// turn UnrealEditor into UnrealEditor-cmd
 		if (!FPaths::GetBaseFilename(ExecutableName).EndsWith("-cmd", ESearchCase::IgnoreCase))
 		{
 			FString NewExeName = ExecutableName + "-Cmd";
diff --git a/Engine/Source/Developer/LauncherServices/Private/Profiles/LauncherProfile.h b/Engine/Source/Developer/LauncherServices/Private/Profiles/LauncherProfile.h
--- a/Engine/Source/Developer/LauncherServices/Private/Profiles/LauncherProfile.h
+++ b/Engine/Source/Developer/LauncherServices/Private/Profiles/LauncherProfile.h
@@ -2331,8 +2331,8 @@ public:
 				FString RelativeProjectPath = Path;
 				bool bRelative = FPaths::MakePathRelativeTo(RelativeProjectPath, *FPaths::RootDir());
 
-				bool bIsUnderUE4Root = bRelative && !(RelativeProjectPath.StartsWith(FString("../"), ESearchCase::CaseSensitive));
-				if (bIsUnderUE4Root)
+				bool bIsUnderUERoot = bRelative && !(RelativeProjectPath.StartsWith(FString("../"), ESearchCase::CaseSensitive));
+				if (bIsUnderUERoot)
 				{
 					ShareableProjectPath = RelativeProjectPath;
 				}
@@ -2900,7 +2900,7 @@ private:
 	// Holds the full absolute path to the Unreal project used by this profile.
 	FString FullProjectPath;
 
-	// Holds the path that might be shareable between people.  Only works if the project is under the UE4 root.
+	// Holds the path that might be shareable between people.  Only works if the project is under the UE root.
 	// otherwise this is an absolute path.
 	FString ShareableProjectPath;
 
