Set ShareableProjectPath in FLauncherProfile::SetDefaults

#rnx

diff --git a/Engine/Source/Developer/LauncherServices/Private/Profiles/LauncherProfile.h b/Engine/Source/Developer/LauncherServices/Private/Profiles/LauncherProfile.h
--- a/Engine/Source/Developer/LauncherServices/Private/Profiles/LauncherProfile.h
+++ b/Engine/Source/Developer/LauncherServices/Private/Profiles/LauncherProfile.h
@@ -196,7 +196,7 @@ public:
 		// profiles created for your persistent devices to be in debug. The user might not see this if they don't expand the Advanced options.
 		BuildConfiguration = EBuildConfiguration::Development;
 		
-		CookMode = ELauncherProfileCookModes::OnTheFly;		
+		CookMode = ELauncherProfileCookModes::OnTheFly;
 	}
 
 private:
@@ -1875,6 +1875,19 @@ public:
 			FullProjectPath = FString();
 		}
 
+		FString RelativeProjectPath = FullProjectPath;
+		bool bRelative = FPaths::MakePathRelativeTo(RelativeProjectPath, *FPaths::RootDir());
+
+		bool bIsUnderUERoot = bRelative && !(RelativeProjectPath.StartsWith(FString("../"), ESearchCase::CaseSensitive));
+		if (bIsUnderUERoot)
+		{
+			ShareableProjectPath = RelativeProjectPath;
+		}
+		else
+		{
+			ShareableProjectPath = FullProjectPath;
+		}
+
 		// Use the locally specified project path is resolving through the root isn't working
 		ProjectSpecified = !FullProjectPath.IsEmpty();
 		
