Editor launcher cleanup.
Initialize all members in FLauncherTask, and move ReadPipe and WritePipe to FLauncherUATTask.
Initialize all members in FLauncherWorker.
Delete FLauncherVerifyProfileTask since Profile->IsValidForLaunch() is aready checked in FLauncher::Launch before creating the FLauncherWorker.
Move duplicate parsing of "Running Package@Device" and device package handling from all separate wait tasks to the FlauncherWorker.
Move handling of killing the uat process tree from all separate wait tasks to the FLauncherWorker.
Change the launch wait task name from "Build Task" to "Launch Task" to correctly display "Launching..." instead of "Validating/Building executable for..."

#rnx

diff --git a/Engine/Plugins/Runtime/Oculus/OculusVR/Source/OculusEditor/Private/OculusBuildAnalytics.cpp b/Engine/Plugins/Runtime/Oculus/OculusVR/Source/OculusEditor/Private/OculusBuildAnalytics.cpp
--- a/Engine/Plugins/Runtime/Oculus/OculusVR/Source/OculusEditor/Private/OculusBuildAnalytics.cpp
+++ b/Engine/Plugins/Runtime/Oculus/OculusVR/Source/OculusEditor/Private/OculusBuildAnalytics.cpp
@@ -193,11 +193,11 @@ void FOculusBuildAnalytics::OnStageStarted(const FString& StageName)
 	{
 		CurrentBuildStage = COOK_IN_EDITOR_STAGE;
 	}
-	else if (StageName.Equals("Build Task") && CurrentBuildStage == LAUNCH_UAT_STAGE)
+	else if (StageName.Equals("Build Task"))
 	{
 		CurrentBuildStage = COMPILE_STAGE;
 	}
-	else if (StageName.Equals("Build Task"))
+	else if (StageName.Equals("Launch Task"))
 	{
 		CurrentBuildStage = LAUNCH_UAT_STAGE;
 	}
diff --git a/Engine/Source/Developer/LauncherServices/Private/Launcher/LauncherTask.h b/Engine/Source/Developer/LauncherServices/Private/Launcher/LauncherTask.h
--- a/Engine/Source/Developer/LauncherServices/Private/Launcher/LauncherTask.h
+++ b/Engine/Source/Developer/LauncherServices/Private/Launcher/LauncherTask.h
@@ -25,16 +25,9 @@ public:
 	 * @param InName - The name of the task.
 	 * @param InDependency - An optional task that this task depends on.
 	 */
-	FLauncherTask( const FString& InName, const FString& InDesc, void* InReadPipe, void* InWritePipe)
+	FLauncherTask( const FString& InName, const FString& InDesc)
 		: Name(InName)
 		, Desc(InDesc)
-		, Status(ELauncherTaskStatus::Pending)
-		, bCancelling( false)
-		, ReadPipe(InReadPipe)
-		, WritePipe(InWritePipe)
-		, Result(0)
-		, ErrorCounter(0)
-		, WarningCounter(0)
 	{ }
 
 public:
@@ -208,12 +201,9 @@ public:
 		// this can be called by anyone
 		bCancelling = true;
 
-		if (Status == ELauncherTaskStatus::Pending || Status == ELauncherTaskStatus::Completed)
+		if (Status == ELauncherTaskStatus::Pending)
 		{
-			if (Status == ELauncherTaskStatus::Pending)
-			{
-				Status = ELauncherTaskStatus::Canceled;
-			}
+			Status = ELauncherTaskStatus::Canceled;
 		}
 		CancelContinuations();
 	}
@@ -335,30 +325,25 @@ private:
 	FDateTime StartTime;
 
 	// Holds the status of this task
-	ELauncherTaskStatus::Type Status;
+	ELauncherTaskStatus::Type Status = ELauncherTaskStatus::Pending;
 
 	// set if this should be cancelled
-	bool bCancelling;
-
+	bool bCancelling = false;
 
 	// Holds the thread that's running this task.
-	FRunnableThread* Thread;
+	FRunnableThread* Thread = nullptr;
 
 	// message delegates
 	FOnTaskStartedDelegate TaskStarted;
 	FOnTaskCompletedDelegate TaskCompleted;
 
 protected:
-	// read and write pipe
-	void* ReadPipe;
-	void* WritePipe;
-
 	// result
-	int32 Result;
+	int32 Result = 0;
 
 	// counters
-	uint32 ErrorCounter;
-	uint32 WarningCounter;
+	uint32 ErrorCounter = 0;
+	uint32 WarningCounter = 0;
 
 	// task counter, used to generate unique thread names for each task
 	static FThreadSafeCounter TaskCounter;
diff --git a/Engine/Source/Developer/LauncherServices/Private/Launcher/LauncherUATTask.h b/Engine/Source/Developer/LauncherServices/Private/Launcher/LauncherUATTask.h
--- a/Engine/Source/Developer/LauncherServices/Private/Launcher/LauncherUATTask.h
+++ b/Engine/Source/Developer/LauncherServices/Private/Launcher/LauncherUATTask.h
@@ -18,14 +18,16 @@ class FLauncherUATTask
 public:
 
 	FLauncherUATTask( const FString& InCommandLine, const FString& InName, const FString& InDesc, void* InReadPipe, void* InWritePipe, const FString& InEditorExe, FProcHandle& InProcessHandle, ILauncherWorker* InWorker, const FString& InCommandEnd, const FString& InBuildCookRunPreCommand)
-		: FLauncherTask(InName, InDesc, InReadPipe, InWritePipe)
+		: FLauncherTask(InName, InDesc)
+
 		, CommandLine(InCommandLine)
+		, ReadPipe(InReadPipe)
+		, WritePipe(InWritePipe)
 		, EditorExe(InEditorExe)
 		, ProcessHandle(InProcessHandle)
 		, CommandText(InCommandEnd)
 		, PreCommand(InBuildCookRunPreCommand)
 	{
-		EndTextFound = false;
 		InWorker->OnOutputReceived().AddRaw(this, &FLauncherUATTask::HandleOutputReceived);
 	}
 
@@ -85,13 +87,6 @@ protected:
 
 		while (FPlatformProcess::IsProcRunning(ProcessHandle) && !EndTextFound)
 		{
-
-			if (IsCancelling())
-			{
-				FPlatformProcess::TerminateProc(ProcessHandle, true);
-				return false;
-			}
-
 			FPlatformProcess::Sleep(0.25);
 		}
 
@@ -122,6 +117,10 @@ private:
 	// Command line
 	FString CommandLine;
 
+	// read and write pipe
+	void* ReadPipe = nullptr;
+	void* WritePipe = nullptr;
+
 	// The editor executable that UAT should use
 	FString EditorExe;
 
@@ -134,6 +133,6 @@ private:
 	// a UAT command to run before BuildCookRun (in the same instance as the BuildCookRun)
 	FString PreCommand;
 
-	bool EndTextFound;
+	bool EndTextFound = false;
 };
 
diff --git a/Engine/Source/Developer/LauncherServices/Private/Launcher/LauncherVerifyProfileTask.h b/Engine/Source/Developer/LauncherServices/Private/Launcher/LauncherVerifyProfileTask.h
deleted file mode 100644
diff --git a/Engine/Source/Developer/LauncherServices/Private/Launcher/LauncherWorker.cpp b/Engine/Source/Developer/LauncherServices/Private/Launcher/LauncherWorker.cpp
--- a/Engine/Source/Developer/LauncherServices/Private/Launcher/LauncherWorker.cpp
+++ b/Engine/Source/Developer/LauncherServices/Private/Launcher/LauncherWorker.cpp
@@ -16,7 +16,6 @@
 #include "Launcher/LauncherTaskChainState.h"
 #include "Launcher/LauncherTask.h"
 #include "Launcher/LauncherUATTask.h"
-#include "Launcher/LauncherVerifyProfileTask.h"
 #include "PlatformInfo.h"
 #include "Misc/ConfigCacheIni.h"
 #include "Profiles/LauncherProfile.h"
@@ -38,7 +37,6 @@ FThreadSafeCounter FLauncherTask::TaskCounter;
 FLauncherWorker::FLauncherWorker(const TSharedRef<ITargetDeviceProxyManager>& InDeviceProxyManager, const ILauncherProfileRef& InProfile)
 	: DeviceProxyManager(InDeviceProxyManager)
 	, Profile(InProfile)
-	, Status(ELauncherWorkerStatus::Busy)
 {
 	CreateAndExecuteTasks(InProfile);
 }
@@ -59,6 +57,26 @@ uint32 FLauncherWorker::Run( )
 
 	LaunchStartTime = FPlatformTime::Seconds();
 
+	auto MessageReceived = [this](const FString& InMessage)
+	{
+		FStringView MessageView = InMessage;
+		{
+			FStringView PackageDevicePrefix = TEXTVIEW("Running Package@Device:");
+			if (MessageView.StartsWith(PackageDevicePrefix))
+			{
+				FStringView Value = MessageView.RightChop(PackageDevicePrefix.Len());
+				int32 SplitIndex;
+				if (Value.FindChar('@', SplitIndex))
+				{
+					FString Package(Value.SubStr(0, SplitIndex));
+					FString Device(Value.SubStr(SplitIndex + 1, Value.Len()));
+					AddDevicePackagePair(Device, Package);
+				}
+			}
+		}
+		OutputMessageReceived.Broadcast(InMessage);
+	};
+
 	// wait for tasks to be completed
 	while (Status == ELauncherWorkerStatus::Busy)
 	{
@@ -76,7 +94,7 @@ uint32 FLauncherWorker::Run( )
 				for (int32 Index = 0; Index < count-1; ++Index)
 				{
 					StringArray[Index].TrimEndInline();
-					OutputMessageReceived.Broadcast(StringArray[Index]);
+					MessageReceived(StringArray[Index]);
 				}
 				Line = StringArray[count-1];
 				if (NewLine.EndsWith(TEXT("\n")))
@@ -102,7 +120,7 @@ uint32 FLauncherWorker::Run( )
 					for (int32 Index = 0; Index < count-1; ++Index)
 					{
 						StringArray[Index].TrimEndInline();
-						OutputMessageReceived.Broadcast(StringArray[Index]);
+						MessageReceived(StringArray[Index]);
 					}
 					Line = StringArray[count-1];
 					if (NewLine.EndsWith(TEXT("\n")))
@@ -115,7 +133,7 @@ uint32 FLauncherWorker::Run( )
 			}
 
 			// fire off the last line
-			OutputMessageReceived.Broadcast(Line);
+			MessageReceived(Line);
 
 		}
 	}
@@ -167,6 +185,9 @@ void FLauncherWorker::Cancel( )
 	if (Status == ELauncherWorkerStatus::Busy)
 	{
 		Status = ELauncherWorkerStatus::Canceling;
+		// kill the uat process tree
+		FPlatformProcess::TerminateProc(ProcHandle, true);
+		// kill any lingering target processes left after killing uat
 		TerminateLaunchedProcess();
 	}
 }
@@ -931,7 +952,6 @@ void FLauncherWorker::CreateAndExecuteTasks( const ILauncherProfileRef& InProfil
 	FPlatformProcess::CreatePipe(ReadPipe, WritePipe);
 
 	// create task chains
-	TaskChain = MakeShareable(new FLauncherVerifyProfileTask());
 	TArray<FString> Platforms;
 	if (InProfile->GetCookMode() == ELauncherProfileCookModes::ByTheBook || InProfile->ShouldBuild())
 	{
@@ -967,14 +987,30 @@ void FLauncherWorker::CreateAndExecuteTasks( const ILauncherProfileRef& InProfil
 	check( InProfile->GetCookMode() != ELauncherProfileCookModes::OnTheFlyInEditor );
 #endif
 
-	TSharedPtr<FLauncherTask> NextTask = TaskChain;
+	TSharedPtr<FLauncherTask> NextTask;
+	auto AddTask = [this, &NextTask](TSharedPtr<FLauncherTask> NewTask)
+	{
+		NewTask->OnStarted().AddRaw(this, &FLauncherWorker::OnTaskStarted);
+		NewTask->OnCompleted().AddRaw(this, &FLauncherWorker::OnTaskCompleted);
+
+		if (NextTask.IsValid())
+		{
+			NextTask->AddContinuation(NewTask);
+			NextTask = NewTask;
+		}
+		else
+		{
+			NextTask = TaskChain = NewTask;
+		}
+	};
+
 	if (InProfile->GetCookMode() == ELauncherProfileCookModes::ByTheBookInEditor)
 	{
 		// need a command which will wait for the cook to finish
 		class FWaitForCookInEditorToFinish : public FLauncherTask
 		{
 		public:
-			FWaitForCookInEditorToFinish() : FLauncherTask( FString(TEXT("Cooking in the editor")), FString(TEXT("Prepairing content to run on device")), NULL, NULL)
+			FWaitForCookInEditorToFinish() : FLauncherTask( FString(TEXT("Cooking in the editor")), FString(TEXT("Prepairing content to run on device")))
 			{
 			}
 			virtual bool PerformTask( FLauncherTaskChainState& ChainState ) override
@@ -992,10 +1028,7 @@ void FLauncherWorker::CreateAndExecuteTasks( const ILauncherProfileRef& InProfil
 			}
 		};
 		TSharedPtr<FLauncherTask> WaitTask = MakeShareable(new FWaitForCookInEditorToFinish());
-		WaitTask->OnStarted().AddRaw(this, &FLauncherWorker::OnTaskStarted);
-		WaitTask->OnCompleted().AddRaw(this, &FLauncherWorker::OnTaskCompleted);
-		NextTask->AddContinuation(WaitTask);
-		NextTask = WaitTask;
+		AddTask(WaitTask);
 	}
 	TArray<FCommandDesc> Commands;
 	FString StartString;
@@ -1008,36 +1041,32 @@ void FLauncherWorker::CreateAndExecuteTasks( const ILauncherProfileRef& InProfil
 		TurnkeyCommand = FString::Printf(TEXT("Turnkey -command=VerifySdk -type=Flash -device=%s -UpdateIfNeeded -utf8output -WaitForUATMutex %s"), *FString::Join(DeviceGroup->GetDeviceIDs(), TEXT("+")), *ITurnkeyIOModule::Get().GetUATParams());
 	}
 
-	TSharedPtr<FLauncherTask> BuildTask = MakeShareable(new FLauncherUATTask(UATCommand, TEXT("Build Task"), TEXT("Launching UAT..."), ReadPipe, WritePipe, InProfile->GetEditorExe(), ProcHandle, this, StartString, TurnkeyCommand));
-	BuildTask->OnStarted().AddRaw(this, &FLauncherWorker::OnTaskStarted);
-	BuildTask->OnCompleted().AddRaw(this, &FLauncherWorker::OnTaskCompleted);
-	NextTask->AddContinuation(BuildTask);
-	NextTask = BuildTask;
+	TSharedPtr<FLauncherTask> LaunchTask = MakeShareable(new FLauncherUATTask(UATCommand, TEXT("Launch Task"), TEXT("Launching UAT..."), ReadPipe, WritePipe, InProfile->GetEditorExe(), ProcHandle, this, StartString, TurnkeyCommand));
+	AddTask(LaunchTask);
 	for (int32 Index = 0; Index < Commands.Num(); ++Index)
 	{
 		class FLauncherWaitTask : public FLauncherTask
 		{
 		public:
 			FLauncherWaitTask( const FString& InCommandEnd, const FString& InName, const FString& InDesc, FProcHandle& InProcessHandle, ILauncherWorker* InWorker)
-				: FLauncherTask(InName, InDesc, 0, 0)
+				: FLauncherTask(InName, InDesc)
 				, CommandText(InCommandEnd)
 				, ProcessHandle(InProcessHandle)
 				, LauncherWorker(InWorker)
 			{
-				EndTextFound = false;
 				InWorker->OnOutputReceived().AddRaw(this, &FLauncherWaitTask::HandleOutputReceived);
 			}
 
+			virtual void Exit()
+			{
+				LauncherWorker->OnOutputReceived().RemoveAll(this);
+			}
+
 		protected:
 			virtual bool PerformTask( FLauncherTaskChainState& ChainState ) override
 			{
 				while (FPlatformProcess::IsProcRunning(ProcessHandle) && !EndTextFound)
 				{
-					if (IsCancelling())
-					{
-						FPlatformProcess::TerminateProc(ProcessHandle, true);
-						return false;
-					}
 					FPlatformProcess::Sleep(0.25);
 				}
 				if (!EndTextFound && !FPlatformProcess::GetProcReturnCode(ProcessHandle, &Result))
@@ -1050,39 +1079,23 @@ void FLauncherWorker::CreateAndExecuteTasks( const ILauncherProfileRef& InProfil
 			void HandleOutputReceived(const FString& InMessage)
 			{
 				EndTextFound |= InMessage.Contains(CommandText);
-				const FString DevicePackagePairMessagePrefix = "Running Package@Device:";
-				if (InMessage.StartsWith(DevicePackagePairMessagePrefix))
-				{
-					FString DevicePackagePairMessage = InMessage;
-					DevicePackagePairMessage.RemoveFromStart(DevicePackagePairMessagePrefix);
-					TArray<FString> DevicePackagePair;
-					if (DevicePackagePairMessage.ParseIntoArray(DevicePackagePair, TEXT("@"), true) == 2)
-					{
-						LauncherWorker->AddDevicePackagePair(DevicePackagePair[1], DevicePackagePair[0]);
-					}
-				}
 			}
 
 		private:
 			FString CommandText;
 			FProcHandle& ProcessHandle;
-			bool EndTextFound;
-			ILauncherWorker* LauncherWorker;
+			ILauncherWorker* LauncherWorker = nullptr;
+			bool EndTextFound = false;
 		};			
 
 		TSharedPtr<FLauncherTask> WaitTask = MakeShareable(new FLauncherWaitTask(Commands[Index].EndText, Commands[Index].Name, Commands[Index].Desc, ProcHandle, this));
-		WaitTask->OnStarted().AddRaw(this, &FLauncherWorker::OnTaskStarted);
-		WaitTask->OnCompleted().AddRaw(this, &FLauncherWorker::OnTaskCompleted);
-		NextTask->AddContinuation(WaitTask);
-		NextTask = WaitTask;
+		AddTask(WaitTask);
 	}
 
 	// execute the chain
 	FLauncherTaskChainState ChainState;
-
 	ChainState.Profile = InProfile;
 	ChainState.SessionId = FGuid::NewGuid();
-
 	TaskChain->Execute(ChainState);
 }
 
diff --git a/Engine/Source/Developer/LauncherServices/Private/Launcher/LauncherWorker.h b/Engine/Source/Developer/LauncherServices/Private/Launcher/LauncherWorker.h
--- a/Engine/Source/Developer/LauncherServices/Private/Launcher/LauncherWorker.h
+++ b/Engine/Source/Developer/LauncherServices/Private/Launcher/LauncherWorker.h
@@ -140,18 +140,18 @@ private:
 	ILauncherProfilePtr Profile;
 
 	// Holds the worker's current status.
-	volatile ELauncherWorkerStatus::Type Status;
+	volatile ELauncherWorkerStatus::Type Status = ELauncherWorkerStatus::Busy;
 
 	// Holds the first task in the task chain.
 	TSharedPtr<FLauncherTask> TaskChain;
 
 	// holds the read and write pipes and the running UAT process
-	void* ReadPipe;
-	void* WritePipe;
+	void* ReadPipe = nullptr;
+	void* WritePipe = nullptr;
 	FProcHandle ProcHandle;
 
-	double StageStartTime;
-	double LaunchStartTime;
+	double StageStartTime = 0.0;
+	double LaunchStartTime = 0.0;
 
 	// message delegate
 	FOutputMessageReceivedDelegate OutputMessageReceived;
