UAT: For editor quick launch and packaging, only pass -iostore and -compressed when using -pak, and never pass -compressed to the cook commandlet (that feature has been removed).

#rnx

diff --git a/Engine/Source/Developer/LauncherServices/Private/Launcher/LauncherWorker.cpp b/Engine/Source/Developer/LauncherServices/Private/Launcher/LauncherWorker.cpp
--- a/Engine/Source/Developer/LauncherServices/Private/Launcher/LauncherWorker.cpp
+++ b/Engine/Source/Developer/LauncherServices/Private/Launcher/LauncherWorker.cpp
@@ -635,11 +635,14 @@ FString FLauncherWorker::CreateUATCommand( const ILauncherProfileRef& InProfile,
 			if (InProfile->IsPackingWithUnrealPak())
 			{
 				UATCommand += TEXT(" -pak");
-			}
-
-			if (InProfile->IsUsingIoStore())
-			{
-				UATCommand += TEXT(" -iostore");
+				if (InProfile->IsUsingIoStore())
+				{
+					UATCommand += TEXT(" -iostore");
+				}
+				if (InProfile->IsCompressed())
+				{
+					UATCommand += TEXT(" -compressed");
+				}
 			}
 
 			if (InProfile->MakeBinaryConfig())
@@ -782,10 +785,14 @@ FString FLauncherWorker::CreateUATCommand( const ILauncherProfileRef& InProfile,
 		if (InProfile->IsPackingWithUnrealPak())
 		{
 			UATCommand += TEXT(" -pak");
-		}
-		if (InProfile->IsUsingIoStore())
-		{
-			UATCommand += TEXT(" -iostore");
+			if (InProfile->IsUsingIoStore())
+			{
+				UATCommand += TEXT(" -iostore");
+			}
+			if (InProfile->IsCompressed())
+			{
+				UATCommand += TEXT(" -compressed");
+			}
 		}
 		break;
 	case ELauncherProfileCookModes::DoNotCook:
@@ -813,11 +820,6 @@ FString FLauncherWorker::CreateUATCommand( const ILauncherProfileRef& InProfile,
 		UATCommand += TEXT(" -SkipCookingEditorContent");
 	}
 
-	if ( InProfile->IsCompressed() )
-	{
-		UATCommand += TEXT(" -compressed");
-	}
-
 	// stage/package/deploy
 	if (InProfile->GetDeploymentMode() != ELauncherProfileDeploymentModes::DoNotDeploy)
 	{
diff --git a/Engine/Source/Editor/TurnkeySupport/Private/TurnkeySupportModule.cpp b/Engine/Source/Editor/TurnkeySupport/Private/TurnkeySupportModule.cpp
--- a/Engine/Source/Editor/TurnkeySupport/Private/TurnkeySupportModule.cpp
+++ b/Engine/Source/Editor/TurnkeySupport/Private/TurnkeySupportModule.cpp
@@ -418,22 +418,21 @@ public:
 				BuildCookRunParams += TEXT(" -clean");
 			}
 
-			if (PackagingSettings->bCompressed)
-			{
-				BuildCookRunParams += TEXT(" -compressed");
-			}
-
-			if (PackagingSettings->bUseIoStore)
-			{
-				BuildCookRunParams += TEXT(" -iostore");
-
-				// Pak file(s) must be used when using container file(s)
-				PackagingSettings->UsePakFile = true;
-			}
+			// Pak file(s) must be used when using container file(s)
+			PackagingSettings->UsePakFile |= PackagingSettings->bUseIoStore;
 
 			if (PackagingSettings->UsePakFile)
 			{
 				BuildCookRunParams += TEXT(" -pak");
+				if (PackagingSettings->bUseIoStore)
+				{
+					BuildCookRunParams += TEXT(" -iostore");
+				}
+
+				if (PackagingSettings->bCompressed)
+				{
+					BuildCookRunParams += TEXT(" -compressed");
+				}
 			}
 
 			if (PackagingSettings->IncludePrerequisites)
diff --git a/Engine/Source/Programs/AutomationTool/Scripts/CookCommand.Automation.cs b/Engine/Source/Programs/AutomationTool/Scripts/CookCommand.Automation.cs
--- a/Engine/Source/Programs/AutomationTool/Scripts/CookCommand.Automation.cs
+++ b/Engine/Source/Programs/AutomationTool/Scripts/CookCommand.Automation.cs
@@ -222,13 +222,6 @@ namespace AutomationScripts
 				CommandletParams += " -basedonreleaseversionroot=" + Params.BasedOnReleaseVersionBasePath;
 			}
 
-			// if we are not going to pak but we specified compressed then compress in the cooker ;)
-			// otherwise compress the pak files
-			if (!Params.Pak && !Params.SkipPak && Params.Compressed && !Params.ForceUncompressed)
-			{
-				CommandletParams += " -compressed";
-			}
-
 			if (!Params.NoClient)
 			{
 				var MapsList = Maps == null ? new List<string>() : Maps.ToList();
