Dont pass -noclient if we are doing both a client and server build

diff --git a/Engine/Source/Developer/LauncherServices/Private/Launcher/LauncherWorker.cpp b/Engine/Source/Developer/LauncherServices/Private/Launcher/LauncherWorker.cpp
--- a/Engine/Source/Developer/LauncherServices/Private/Launcher/LauncherWorker.cpp
+++ b/Engine/Source/Developer/LauncherServices/Private/Launcher/LauncherWorker.cpp
@@ -420,6 +420,13 @@ FString FLauncherWorker::CreateUATCommand( const ILauncherProfileRef& InProfile,
 			bUATClosesAfterLaunch |= PlatformInfo->DataDrivenPlatformInfo->bUATClosesAfterLaunch;
 		}
 	}
+
+	// If both Client and Server are desired to be built avoid Server causing clients to not be built PlatformInfo wise
+	if (OptionalParams.Contains(TEXT("-client")) && OptionalParams.Contains(TEXT("-noclient")))
+	{
+		OptionalParams = OptionalParams.Replace(TEXT("-noclient"), TEXT(""));
+	}
+
 	if (ServerPlatforms.Len() > 0)
 	{
 		ServerCommand = TEXT(" -server -serverplatform=") + ServerPlatforms.RightChop(1);
diff --git a/Engine/Source/Developer/LauncherServices/Private/Profiles/LauncherProfile.h b/Engine/Source/Developer/LauncherServices/Private/Profiles/LauncherProfile.h
--- a/Engine/Source/Developer/LauncherServices/Private/Profiles/LauncherProfile.h
+++ b/Engine/Source/Developer/LauncherServices/Private/Profiles/LauncherProfile.h
@@ -1680,6 +1680,13 @@ public:
 
 			bUATClosesAfterLaunch |= PlatformInfo->DataDrivenPlatformInfo->bUATClosesAfterLaunch;
 		}
+
+		// If both Client and Server are desired to be built avoid Server causing clients to not be built PlatformInfo wise
+		if (OptionalParams.Contains(TEXT("-client")) && OptionalParams.Contains(TEXT("-noclient")))
+		{
+			OptionalParams = OptionalParams.Replace(TEXT("-noclient"), TEXT(""));
+		}
+
 		return bUATClosesAfterLaunch;
 	}
 
