Fixing JSON warnings with reference container values in project launcher.

diff --git a/Engine/Source/Developer/LauncherServices/Private/Profiles/LauncherProfile.h b/Engine/Source/Developer/LauncherServices/Private/Profiles/LauncherProfile.h
--- a/Engine/Source/Developer/LauncherServices/Private/Profiles/LauncherProfile.h
+++ b/Engine/Source/Developer/LauncherServices/Private/Profiles/LauncherProfile.h
@@ -1832,8 +1832,18 @@ public:
 		CreateReleaseVersion = Object.GetBoolField("CreateReleaseVersion");
 		CreateReleaseVersionName = Object.GetStringField("CreateReleaseVersionName");
 		BasedOnReleaseVersionName = Object.GetStringField("BasedOnReleaseVersionName");
-		ReferenceContainerCryptoKeysFileName = Object.GetStringField("ReferenceContainerCryptoKeysFileName");
-		ReferenceContainerGlobalFileName = Object.GetStringField("ReferenceContainerGlobalFileName");
+
+		if (Version >= LAUNCHERSERVICES_ADDEDREFERENCECONTAINERS)
+		{
+			ReferenceContainerCryptoKeysFileName = Object.GetStringField("ReferenceContainerCryptoKeysFileName");
+			ReferenceContainerGlobalFileName = Object.GetStringField("ReferenceContainerGlobalFileName");
+		}
+		else
+		{
+			ReferenceContainerCryptoKeysFileName.Empty();
+			ReferenceContainerGlobalFileName.Empty();
+		}
+
 		CreateDLC = Object.GetBoolField("CreateDLC");
 		DLCName = Object.GetStringField("DLCName");
 		bGenerateChunks = Object.GetBoolField("GenerateChunks");
