Removing GetNumCookersToSpawn as it hasn't been used in a long long time.

diff --git a/Engine/Source/Developer/LauncherServices/Private/Launcher/LauncherWorker.cpp b/Engine/Source/Developer/LauncherServices/Private/Launcher/LauncherWorker.cpp
--- a/Engine/Source/Developer/LauncherServices/Private/Launcher/LauncherWorker.cpp
+++ b/Engine/Source/Developer/LauncherServices/Private/Launcher/LauncherWorker.cpp
@@ -726,11 +726,6 @@ FString FLauncherWorker::CreateUATCommand( const ILauncherProfileRef& InProfile,
 				UATCommand += TEXT(" -stage");
 			}
 
-			if (InProfile->GetNumCookersToSpawn() > 0)
-			{
-				UATCommand += FString::Printf(TEXT(" -NumCookersToSpawn=%d"), InProfile->GetNumCookersToSpawn());
-			}
-
 			FCommandDesc Desc;
 			FText Command = FText::Format(LOCTEXT("LauncherCookDesc", "Cook content for {0}"), FText::FromString(Platforms.RightChop(1)));
 			Desc.Name = "Cook Task";
diff --git a/Engine/Source/Developer/LauncherServices/Private/Profiles/LauncherProfile.h b/Engine/Source/Developer/LauncherServices/Private/Profiles/LauncherProfile.h
--- a/Engine/Source/Developer/LauncherServices/Private/Profiles/LauncherProfile.h
+++ b/Engine/Source/Developer/LauncherServices/Private/Profiles/LauncherProfile.h
@@ -50,6 +50,7 @@ enum ELauncherVersion
 	LAUNCHERSERVICES_ADDEDUSEIOSTORE = 29,
 	LAUNCHERSERVICES_ADDEDMAKEBINARYCONFIG = 30,
 	LAUNCHERSERVICES_ADDEDREFERENCECONTAINERS = 31,
+	LAUNCHERSERVICES_REMOVEDNUMCOOKERSTOSPAWN = 32,
 	//ADD NEW STUFF HERE
 
 
@@ -420,11 +421,6 @@ public:
 		return CookedCultures;
 	}
 
-	virtual const int32 GetNumCookersToSpawn() const override
-	{
-		return NumCookersToSpawn;
-	}
-
 	virtual const bool GetSkipCookingEditorContent() const override
 	{
 		return bSkipCookingEditorContent;
@@ -969,9 +965,10 @@ public:
 		{
 			Archive << DeployPlatformString;
 		}		
-		if (Version >= LAUNCHERSERVICES_ADDEDNUMCOOKERSTOSPAWN)
+		if (Version >= LAUNCHERSERVICES_ADDEDNUMCOOKERSTOSPAWN && Version < LAUNCHERSERVICES_REMOVEDNUMCOOKERSTOSPAWN)
 		{
-			Archive << NumCookersToSpawn;
+			int32 OldNumCookersToSpawn;
+			Archive << OldNumCookersToSpawn;
 		}
 		if (Version >= LAUNCHERSERVICES_ADDEDSKIPCOOKINGEDITORCONTENT)
 		{
@@ -1172,7 +1169,6 @@ public:
 		Writer.WriteValue("EncryptIniFiles", EncryptIniFiles);
 		Writer.WriteValue("ForDistribution", ForDistribution);
 		Writer.WriteValue("DeployPlatform", DefaultDeployPlatform.ToString());
-		Writer.WriteValue("NumCookersToSpawn", NumCookersToSpawn);
 		Writer.WriteValue("SkipCookingEditorContent", bSkipCookingEditorContent);
 		Writer.WriteValue("DeployIncremental", DeployIncremental);
 		Writer.WriteValue("GeneratePatch", GeneratePatch);
@@ -1424,11 +1420,6 @@ public:
 					Writer.WriteValue("archivedirectory", GetArchiveDirectory());
 				}
 
-				if (GetNumCookersToSpawn() > 0)
-				{
-					Writer.WriteValue("numcookerstospawn", GetNumCookersToSpawn());
-				}
-
 				TMap<FString, FString> CookCommands = ParseCommands(GetCookOptions());
 				for (TMap<FString, FString>::TIterator Iter = CookCommands.CreateIterator(); Iter; ++Iter)
 				{
@@ -1552,7 +1543,6 @@ public:
 		"cookontheflystreaming", "true/false"
 		"unversionedcookcontent", "true/false"
 		"skipcookingeditorcontent", "true/false"
-		"numcookerstospawn", "8"
 		"compressed", "true/false"
 		"iterativecooking", "true/false"
 		"skipcookonthefly", "true/false"
@@ -1812,7 +1802,6 @@ public:
 		}
 
 		DefaultDeployPlatform = *(Object.GetStringField("DeployPlatform"));
-		NumCookersToSpawn = (int32)Object.GetNumberField("NumCookersToSpawn");
 		bSkipCookingEditorContent = Object.GetBoolField("SkipCookingEditorContent");
 		DeployIncremental = Object.GetBoolField("DeployIncremental");
 		GeneratePatch = Object.GetBoolField("GeneratePatch");
@@ -1979,7 +1968,6 @@ public:
 		bSkipCookingEditorContent = false;
         ForceClose = true;
         Timeout = 60;
-		NumCookersToSpawn = 0;
 
 /*		if (GetTargetPlatformManager()->GetRunningTargetPlatform() != NULL)
 		{
@@ -2118,15 +2106,6 @@ public:
 		}
 	}
 
-	virtual void SetNumCookersToSpawn(const int32 InNumCookersToSpawn) override
-	{
-		if (NumCookersToSpawn != InNumCookersToSpawn)
-		{
-			NumCookersToSpawn = InNumCookersToSpawn;
-			Validate();
-		}
-	}
-
 	virtual void SetSkipCookingEditorContent(const bool InSkipCookingEditorContent) override
 	{
 		if (bSkipCookingEditorContent != InSkipCookingEditorContent)
@@ -2865,9 +2844,6 @@ private:
 	// Holds a flag indicating whether packages should be saved without a version.
 	bool CookUnversioned;
 
-	// num cookers we want to spawn during cooking
-	int32 NumCookersToSpawn;
-
 	bool bSkipCookingEditorContent;
 
 	// This setting is used only if cooking by the book (only used if cooking by the book).
diff --git a/Engine/Source/Developer/LauncherServices/Public/ILauncherProfile.h b/Engine/Source/Developer/LauncherServices/Public/ILauncherProfile.h
--- a/Engine/Source/Developer/LauncherServices/Public/ILauncherProfile.h
+++ b/Engine/Source/Developer/LauncherServices/Public/ILauncherProfile.h
@@ -543,22 +543,6 @@ public:
 	virtual const FString& GetCookOptions( ) const = 0;
 
 
-	/**
-	 * Get the number of cookers we want to spawn during cooking
-	 *
-	 * @return number of cookers we want to spawn in addition to the master cooker
-	 */
-	virtual const int32 GetNumCookersToSpawn() const = 0;
-
-	/**
-	 * Set the number of cookers we want to spawn during cooking
-	 * 
-	 * @param InNumCookersToSpawn number of cookers we want to spawn in addition to the master cooker
-	 */
-	virtual void SetNumCookersToSpawn(const int32 InNumCookersToSpawn) = 0; 
-
-
-
 	virtual const bool GetSkipCookingEditorContent() const = 0; 
 
 	/**
diff --git a/Engine/Source/Developer/ProjectLauncher/Private/Widgets/Cook/SProjectLauncherCookByTheBookSettings.cpp b/Engine/Source/Developer/ProjectLauncher/Private/Widgets/Cook/SProjectLauncherCookByTheBookSettings.cpp
--- a/Engine/Source/Developer/ProjectLauncher/Private/Widgets/Cook/SProjectLauncherCookByTheBookSettings.cpp
+++ b/Engine/Source/Developer/ProjectLauncher/Private/Widgets/Cook/SProjectLauncherCookByTheBookSettings.cpp
@@ -661,26 +661,6 @@ TSharedRef<SWidget> SProjectLauncherCookByTheBookSettings::MakeComplexWidget()
 								]
 						]
 
-					// multiprocess cooking options
-					+ SVerticalBox::Slot()
-						.AutoHeight()
-						.Padding(0.0f, 8.0f, 0.0f, 0.0f)
-						[
-							SNew(SProjectLauncherFormLabel)
-								.LabelText(LOCTEXT("MultiProcessCookerTextBoxLabel", "Num cookers to spawn:"))
-						]
-
-					+ SVerticalBox::Slot()
-						.AutoHeight()
-						.Padding(0.0f, 4.0f, 0.0f, 0.0f)
-						[
-							// cooker command line options
-							SNew(SEditableTextBox)
-								.ToolTipText(LOCTEXT("MultiProcessCookerTextBoxTooltip", "The number of cookers to spawn when we do a cook by the book."))
-								.Text(this, &SProjectLauncherCookByTheBookSettings::HandleMultiProcessCookerTextBlockText)
-								.OnTextCommitted(this, &SProjectLauncherCookByTheBookSettings::HandleMultiProcessCookerCommitted)
-						]
-
 					// unreal pak check box
 					+ SVerticalBox::Slot()
 						.AutoHeight()
@@ -1645,42 +1625,6 @@ void SProjectLauncherCookByTheBookSettings::HandleCookerOptionsCommitted(const F
 }
 
 
-FText SProjectLauncherCookByTheBookSettings::HandleMultiProcessCookerTextBlockText() const
-{
-	ILauncherProfilePtr SelectedProfile = Model->GetSelectedProfile();
-
-	FText result;
-
-	if (SelectedProfile.IsValid())
-	{
-		result = FText::FromString(FString::FromInt(SelectedProfile->GetNumCookersToSpawn()));
-	}
-
-	return result;
-}
-
-
-void SProjectLauncherCookByTheBookSettings::HandleMultiProcessCookerCommitted(const FText& NewText, ETextCommit::Type CommitType)
-{
-	ILauncherProfilePtr SelectedProfile = Model->GetSelectedProfile();
-
-	if (SelectedProfile.IsValid())
-	{
-		int32 NumCookersToSpawn = FCString::Atoi(*NewText.ToString());
-		switch (CommitType)
-		{
-		case ETextCommit::Default:
-		case ETextCommit::OnCleared:
-			NumCookersToSpawn = 0;
-			break;
-		default:
-			break;
-		}
-		SelectedProfile->SetNumCookersToSpawn(NumCookersToSpawn);
-	}
-}
-
-
 void SProjectLauncherCookByTheBookSettings::HandleUnrealPakCheckBoxCheckStateChanged(ECheckBoxState NewState)
 {
 	ILauncherProfilePtr SelectedProfile = Model->GetSelectedProfile();
diff --git a/Engine/Source/Developer/ProjectLauncher/Private/Widgets/Cook/SProjectLauncherCookByTheBookSettings.h b/Engine/Source/Developer/ProjectLauncher/Private/Widgets/Cook/SProjectLauncherCookByTheBookSettings.h
--- a/Engine/Source/Developer/ProjectLauncher/Private/Widgets/Cook/SProjectLauncherCookByTheBookSettings.h
+++ b/Engine/Source/Developer/ProjectLauncher/Private/Widgets/Cook/SProjectLauncherCookByTheBookSettings.h
@@ -139,11 +139,6 @@ private:
 	/** Callback for changing the cookers additional options */
 	void HandleCookerOptionsCommitted(const FText& NewText, ETextCommit::Type CommitType);
 
-	/** Callback for when the number of cookers need to be refreshed */
-	FText HandleMultiProcessCookerTextBlockText() const;
-	/** callback for when the number of cookers changes */
-	void HandleMultiProcessCookerCommitted(const FText& NewText, ETextCommit::Type CommitType);
-
 	/** Callback for updating any settings after the selected project has changed in the profile. */
 	void HandleProfileProjectChanged();
 
diff --git a/Engine/Source/Editor/TurnkeySupport/Private/TurnkeyEditorSupport.cpp b/Engine/Source/Editor/TurnkeySupport/Private/TurnkeyEditorSupport.cpp
--- a/Engine/Source/Editor/TurnkeySupport/Private/TurnkeyEditorSupport.cpp
+++ b/Engine/Source/Editor/TurnkeySupport/Private/TurnkeyEditorSupport.cpp
@@ -40,11 +40,6 @@ FString FTurnkeyEditorSupport::GetUATOptions()
 	FString Options;
 	Options += FString::Printf(TEXT(" -unrealexe=\"%s\""), *FUnrealEdMisc::Get().GetExecutableForCommandlets());
 
-	int32 NumCookers = GetDefault<UEditorExperimentalSettings>()->MultiProcessCooking;
-	if (NumCookers > 0)
-	{
-		Options += FString::Printf(TEXT(" -NumCookersToSpawn=%d"), NumCookers);
-	}
 	return Options;
 #else
 	return TEXT("");
diff --git a/Engine/Source/Editor/UnrealEd/Classes/Settings/EditorExperimentalSettings.h b/Engine/Source/Editor/UnrealEd/Classes/Settings/EditorExperimentalSettings.h
--- a/Engine/Source/Editor/UnrealEd/Classes/Settings/EditorExperimentalSettings.h
+++ b/Engine/Source/Editor/UnrealEd/Classes/Settings/EditorExperimentalSettings.h
@@ -87,8 +87,8 @@ public:
 	UPROPERTY(EditAnywhere, config, Category = Cooking, meta = (DisplayName = "Use shared cooked builds in launch on", ConfigRestartRequired = true))
 	bool bSharedCookedBuilds;
 
-	UPROPERTY(EditAnywhere, config, Category = Cooking, meta = (DisplayName = "Use multiple processes when cooking (only affects File -> Package)"))
-	int32 MultiProcessCooking;
+	UPROPERTY()
+	int32 MultiProcessCooking_DEPRECATED;
 
 	/** Enable late joining in PIE */
 	UPROPERTY(EditAnywhere, config, Category = PIE, meta = (DisplayName = "Allow late joining"))
diff --git a/Engine/Source/Programs/AutomationTool/AutomationUtils/ProjectParams.cs b/Engine/Source/Programs/AutomationTool/AutomationUtils/ProjectParams.cs
--- a/Engine/Source/Programs/AutomationTool/AutomationUtils/ProjectParams.cs
+++ b/Engine/Source/Programs/AutomationTool/AutomationUtils/ProjectParams.cs
@@ -334,7 +334,6 @@ namespace AutomationTool
             this.UnversionedCookedContent = InParams.UnversionedCookedContent;
 			this.OptionalContent = InParams.OptionalContent;
 			this.SkipCookingEditorContent = InParams.SkipCookingEditorContent;
-            this.NumCookersToSpawn = InParams.NumCookersToSpawn;
 			this.FileServer = InParams.FileServer;
 			this.DedicatedServer = InParams.DedicatedServer;
 			this.Client = InParams.Client;
@@ -478,7 +477,6 @@ namespace AutomationTool
             bool? EncryptPakIndex = null,
 			bool? EncryptEverything = null,
 			bool? SkipCookingEditorContent = null,
-            int? NumCookersToSpawn = null,
             string AdditionalCookerOptions = null,
 			string OriginalReleaseVersion = null,
 			string BasedOnReleaseVersion = null,
@@ -756,14 +754,6 @@ namespace AutomationTool
             this.UnversionedCookedContent = GetOptionalParamValueIfNotSpecified(Command, UnversionedCookedContent, this.UnversionedCookedContent, "UnversionedCookedContent", "VersionCookedContent");
 			this.OptionalContent = GetOptionalParamValueIfNotSpecified(Command, OptionalContent, this.OptionalContent, "editoroptional", "noeditoroptional");
 			this.SkipCookingEditorContent = GetParamValueIfNotSpecified(Command, SkipCookingEditorContent, this.SkipCookingEditorContent, "SkipCookingEditorContent");
-            if (NumCookersToSpawn.HasValue)
-            {
-                this.NumCookersToSpawn = NumCookersToSpawn.Value;
-            }
-            else if (Command != null)
-            {
-                this.NumCookersToSpawn = Command.ParseParamInt("NumCookersToSpawn");
-            }
 			this.Compressed = GetParamValueIfNotSpecified(Command, Compressed, this.Compressed, "compressed");
 			this.ForceUncompressed = GetParamValueIfNotSpecified(Command, ForceUncompressed, this.ForceUncompressed, "ForceUncompressed");
 			this.AdditionalPakOptions = ParseParamValueIfNotSpecified(Command, AdditionalPakOptions, "AdditionalPakOptions");
diff --git a/Engine/Source/Programs/AutomationTool/Scripts/CookCommand.Automation.cs b/Engine/Source/Programs/AutomationTool/Scripts/CookCommand.Automation.cs
--- a/Engine/Source/Programs/AutomationTool/Scripts/CookCommand.Automation.cs
+++ b/Engine/Source/Programs/AutomationTool/Scripts/CookCommand.Automation.cs
@@ -185,10 +185,6 @@ namespace AutomationScripts
 			{
 				CommandletParams += " -createreleaseversion=" + Params.CreateReleaseVersion;
 			}
-			if (Params.NumCookersToSpawn != 0)
-			{
-				CommandletParams += " -numcookerstospawn=" + Params.NumCookersToSpawn;
-			}
 			if (Params.CookPartialGC)
 			{
 				CommandletParams += " -partialgc";
