Project Launcher supports multiple .target.cs Build Target files.

 - also includes a fix for the wrong 'previous' profile being passed to FProjectLauncherModel::ProfileSelectedDelegate.
 - includes a fix for the wrong project being used for device profiles if the fallback project is changed in the editor.

diff --git a/Engine/Source/Developer/LauncherServices/Private/Profiles/LauncherProfile.h b/Engine/Source/Developer/LauncherServices/Private/Profiles/LauncherProfile.h
--- a/Engine/Source/Developer/LauncherServices/Private/Profiles/LauncherProfile.h
+++ b/Engine/Source/Developer/LauncherServices/Private/Profiles/LauncherProfile.h
@@ -21,6 +21,7 @@
 #include "PlatformInfo.h"
 #include "TargetReceipt.h"
 #include "DesktopPlatformModule.h"
+#include "GameProjectHelper.h"
 
 DEFINE_LOG_CATEGORY_STATIC(LogLauncherProfile, Log, All);
 
@@ -52,6 +53,7 @@ enum ELauncherVersion
 	LAUNCHERSERVICES_ADDEDREFERENCECONTAINERS = 31,
 	LAUNCHERSERVICES_REMOVEDNUMCOOKERSTOSPAWN = 32,
 	LAUNCHERSERVICES_ADDEDORIGINALRELEASEVERSION = 33,
+	LAUNCHERSERVICES_ADDBUILDTARGETNAME = 34,
 	//ADD NEW STUFF HERE
 
 
@@ -247,6 +249,7 @@ public:
 		, DefaultLaunchRole(MakeShareable(new FLauncherProfileLaunchRole()))
 	{ 
 		SetDefaults();
+		ProjectChangedDelegate.AddRaw(this, &FLauncherProfile::OnSelectedProjectChanged);
 	}
 
 	/**
@@ -261,6 +264,7 @@ public:
 		, Name(InProfileName)
 	{
 		SetDefaults();
+		ProjectChangedDelegate.AddRaw(this, &FLauncherProfile::OnSelectedProjectChanged);
 	}
 
 	/**
@@ -309,6 +313,7 @@ public:
 	{
 		CookedPlatforms.AddUnique(PlatformName);
 
+		RefreshValidBuildTargets();
 		Validate();
 	}
 
@@ -377,6 +382,7 @@ public:
 		{
 			CookedPlatforms.Reset();
 
+			RefreshValidBuildTargets();
 			Validate();
 		}
 	}
@@ -397,9 +403,21 @@ public:
 		return BuildConfiguration;
 	}
 
-	virtual const FString& GetBuildTarget() const override
+	virtual bool HasBuildTargetSpecified() const override
 	{
-		return BuildTargetName;
+		return BuildTargetSpecified;
+	}
+
+	virtual FString GetBuildTarget() const override
+	{
+		if (BuildTargetSpecified)
+		{
+			return BuildTargetName;
+		}
+		else
+		{
+			return LauncherProfileManager->GetBuildTarget();
+		}
 	}
 
 	virtual EBuildConfiguration GetCookConfiguration( ) const override
@@ -891,6 +909,7 @@ public:
 	{
 		CookedPlatforms.Remove(PlatformName);
 
+		RefreshValidBuildTargets();
 		Validate();
 	}
 
@@ -1064,6 +1083,12 @@ public:
 			Archive << OriginalReleaseVersionName;
 		}
 
+		if (Version >= LAUNCHERSERVICES_ADDBUILDTARGETNAME)
+		{
+			Archive << BuildTargetSpecified;
+			Archive << BuildTargetName;
+		}
+
 		else if(Archive.IsLoading())
 		{
 			BuildMode = BuildGame ? ELauncherProfileBuildModes::Build : ELauncherProfileBuildModes::DoNotBuild;
@@ -1104,6 +1129,10 @@ public:
 			SetDefaultDeployPlatform(DefaultDeployPlatform);
 		}
 
+		if (Archive.IsLoading())
+		{
+			RefreshValidBuildTargets();
+		}
 		Validate();
 
 		return true;
@@ -1208,6 +1237,8 @@ public:
 		Writer.WriteValue("IncludePrerequisites", IncludePrerequisites);
 		Writer.WriteValue("UseIoStore", bUseIoStore);
 		Writer.WriteValue("MakeBinaryConfig", bMakeBinaryConfig);
+		Writer.WriteValue("BuildTargetSpecified", BuildTargetSpecified);
+		Writer.WriteValue("BuildTargetName", BuildTargetName);
 
 		// serialize the default launch role
 		DefaultLaunchRole->Save(Writer, TEXT("DefaultRole"));
@@ -1906,6 +1937,12 @@ public:
 			bMakeBinaryConfig = Object.GetBoolField("MakeBinaryConfig");
 		}
 
+		if (Version >= LAUNCHERSERVICES_ADDBUILDTARGETNAME)
+		{
+			BuildTargetSpecified = Object.GetBoolField("BuildTargetSpecified");
+			BuildTargetName = Object.GetStringField("BuildTargetName");
+		}
+
 		// load the default launch role
 		TSharedPtr<FJsonObject> Role = Object.GetObjectField("DefaultRole");
 		DefaultLaunchRole->Load(*(Role.Get()));
@@ -1932,6 +1969,7 @@ public:
 			SetDefaultDeployPlatform(DefaultDeployPlatform);
 		}
 
+		RefreshValidBuildTargets();
 		Validate();
 
 		return true;
@@ -1981,6 +2019,7 @@ public:
 		// default build settings
 		BuildMode = ELauncherProfileBuildModes::Auto;
 		BuildUAT = !FApp::GetEngineIsPromotedBuild() && !FApp::IsEngineInstalled();
+		BuildTargetSpecified = true;
 
 		// default cook settings
 		CookConfiguration = FApp::GetBuildConfiguration();
@@ -2054,6 +2093,7 @@ public:
 		bShouldUpdateFlash = false;
 		bMakeBinaryConfig = false;
 
+		RefreshValidBuildTargets();
 		Validate();
 	}
 
@@ -2097,6 +2137,24 @@ public:
 		}
 	}
 
+	virtual void SetBuildTargetSpecified(bool Specified) override
+	{
+		if (BuildTargetSpecified != Specified)
+		{
+			BuildTargetSpecified = Specified;
+			Validate();
+		}
+	}
+
+	virtual void FallbackBuildTargetUpdated() override
+	{
+		if (!HasBuildTargetSpecified())
+		{
+			Validate();
+		}
+	}
+
+
 	virtual void SetBuildTarget( const FString& TargetName ) override
 	{
 		if (BuildTargetName != TargetName)
@@ -2394,6 +2452,7 @@ public:
 		{
 			ProjectSpecified = Specified;
 
+			RefreshValidBuildTargets();
 			Validate();
 
 			ProjectChangedDelegate.Broadcast();
@@ -2404,6 +2463,7 @@ public:
 	{
 		if (!HasProjectSpecified())
 		{
+			RefreshValidBuildTargets();
 			Validate();
 
 			ProjectChangedDelegate.Broadcast();
@@ -2437,6 +2497,7 @@ public:
 			}
 			CookedMaps.Reset();
 
+			RefreshValidBuildTargets();
 			Validate();
 
 			ProjectChangedDelegate.Broadcast();
@@ -2508,6 +2569,11 @@ public:
 		return ProjectChangedDelegate;
 	}
 
+	virtual FOnProfileBuildTargetOptionsChanged& OnBuildTargetOptionsChanged() override
+	{
+		return BuildTargetOptionsChangedDelegate;
+	}
+
 	virtual void SetEditorExe( const FString& InEditorExe ) override
 	{
 		EditorExe = InEditorExe;
@@ -2566,6 +2632,17 @@ public:
 		return bMakeBinaryConfig;
 	}
 
+	virtual TArray<FString> GetExplicitBuildTargetNames() const override
+	{
+		return ExplictBuildTargetNames;
+	}
+
+	virtual bool RequiresExplicitBuildTargetName() const override
+	{
+		return ExplictBuildTargetNames.Num() > 0;
+	}
+
+
 	//~ End ILauncherProfile Interface
 
 protected:
@@ -2724,6 +2801,7 @@ protected:
 			ValidationErrors.Add(ELauncherProfileValidationErrors::IoStoreRequiresPakFiles);
 		}
 
+		ValidateBuildTarget();
 		ValidatePlatformSDKs();
 		ValidateDeviceStatus();
 	}
@@ -2816,6 +2894,135 @@ protected:
 		}
 	}
 
+	void ValidateBuildTarget()
+	{
+		bool bBuildTargetIsRequired = false;
+		bool bBuildTargetCookVariantMismatch = false;
+
+		FString BuildTarget = GetBuildTarget();
+		TSet<EBuildTargetType> CookTargetTypes = GetCookTargetTypes();
+
+		if (HasProjectSpecified())
+		{
+			if (RequiresExplicitBuildTargetName())
+			{
+				if (CookTargetTypes.Num() > 1 || (!BuildTarget.IsEmpty() && !ExplictBuildTargetNames.Contains(BuildTarget) ) )
+				{
+					// can only build the same Variant (Game, Client, etc) as the selected build target
+					bBuildTargetCookVariantMismatch = true;
+				}
+				else if (BuildTarget.IsEmpty())
+				{
+					// multiple .target.cs files defined of the same Variant - need to specify one
+					bBuildTargetIsRequired = true;
+				}
+			}
+		}		
+		else
+		{
+			// this profile is using the fallback project. Need to check all build targets instead
+			bool bBuildTargetIsCookable = false;
+			bool bVariantRequiresBuildTarget = false;
+
+			const TArray<FTargetInfo>& Targets = FDesktopPlatformModule::Get()->GetTargetsForProject(GetProjectPath());
+			for (const FTargetInfo& Target : Targets)
+			{
+				if (CookTargetTypes.IsEmpty() || CookTargetTypes.Contains(Target.Type))
+				{
+					if (BuildTarget.IsEmpty())
+					{
+						if (LauncherProfileManager->GetAllExplicitBuildTargetNames().Contains(Target.Name))
+						{
+							// a currently selected Variant has multiple .target.cs files - need to specify one
+							bBuildTargetIsRequired = true;
+							break;
+						}
+					}
+					else
+					{
+						if (Target.Name == BuildTarget)
+						{
+							bBuildTargetIsCookable = true;
+							break;
+						}
+					}				
+				}
+			}
+
+			if (!BuildTarget.IsEmpty() && !bBuildTargetIsCookable)
+			{
+				// currently build target does not reference a Variant with multiple .target.cs files
+				bBuildTargetCookVariantMismatch = true;
+			}
+		}
+
+		if (bBuildTargetIsRequired)
+		{
+			if (BuildTargetSpecified)
+			{
+				ValidationErrors.Add(ELauncherProfileValidationErrors::BuildTargetIsRequired);
+			}
+			else
+			{
+				ValidationErrors.Add(ELauncherProfileValidationErrors::FallbackBuildTargetIsRequired);
+			}
+		}
+
+		if (bBuildTargetCookVariantMismatch)
+		{
+			ValidationErrors.Add(ELauncherProfileValidationErrors::BuildTargetCookVariantMismatch);
+		}
+
+	}
+
+
+	void RefreshValidBuildTargets()
+	{
+		TArray<FString> LatestExplicitBuildTargetNames;
+
+		// collect the build targets for the current project, filtered to what we are currently wanting to cook. Do not show fallback project's build targets
+		if (HasProjectSpecified())
+		{
+			TSet<EBuildTargetType> CookTargetTypes = GetCookTargetTypes();
+			LatestExplicitBuildTargetNames = FGameProjectHelper::GetExplicitBuildTargetsForProject( GetProjectPath(), &CookTargetTypes );
+		}
+
+		// notify listeners if the explicitly-required build targets have changed
+		if (ExplictBuildTargetNames != LatestExplicitBuildTargetNames)
+		{
+			ExplictBuildTargetNames = LatestExplicitBuildTargetNames;
+			BuildTargetOptionsChangedDelegate.Broadcast();
+		}
+	}
+
+
+	TSet<EBuildTargetType> GetCookTargetTypes() const
+	{
+		TSet<EBuildTargetType> CookTargetTypes;
+		for ( const FString& Variant : GetCookedPlatforms() )
+		{
+			if (Variant.EndsWith(TEXT("Client")))
+			{
+				CookTargetTypes.Add(EBuildTargetType::Client);
+			}
+			else if (Variant.EndsWith(TEXT("Server")))
+			{
+				CookTargetTypes.Add(EBuildTargetType::Server);
+			}
+			else if (Variant.EndsWith(TEXT("Editor")))
+			{
+				CookTargetTypes.Add(EBuildTargetType::Editor);
+			}
+			else
+			{
+				CookTargetTypes.Add(EBuildTargetType::Game);
+			}
+		}
+
+		return MoveTemp(CookTargetTypes);
+	}
+
+
 	void OnLauncherDeviceGroupDeviceAdded(const ILauncherDeviceGroupRef& DeviceGroup, const FString& DeviceId)
 	{
 		if( DeviceGroup == DeployedDeviceGroup )
@@ -2832,6 +3039,14 @@ protected:
 		}
 	}
 
+	void OnSelectedProjectChanged()
+	{
+		RefreshValidBuildTargets();
+
+		BuildTargetName.Empty();
+		BuildTargetSpecified = ProjectSpecified; // if this is for 'Any Project' then it should use the fallback build target
+	}
+
 private:
 
 	//  Holds a reference to the launcher profile manager.
@@ -2840,6 +3055,9 @@ private:
 	// Holds the desired build configuration (only used if creating new builds).
 	EBuildConfiguration BuildConfiguration;
 
+	// Holds a flag indicating whether the build target is specified by this profile.
+	bool BuildTargetSpecified;
+
 	// Holds the name of the target (matching a .cs file) to build. Needed when multiple targets of a type exist
 	FString BuildTargetName;
 
@@ -3052,4 +3270,10 @@ private:
 
 	// Holds a delegate to be invoked when the project has changed
 	FOnProfileProjectChanged ProjectChangedDelegate;
+
+	// Holds a delegate to be invoked when the project build target options have changed
+	FOnProfileBuildTargetOptionsChanged BuildTargetOptionsChangedDelegate;
+
+	// Cached build target options (not serialized)
+	TArray<FString> ExplictBuildTargetNames;
 };
diff --git a/Engine/Source/Developer/LauncherServices/Private/Profiles/LauncherProfileManager.cpp b/Engine/Source/Developer/LauncherServices/Private/Profiles/LauncherProfileManager.cpp
--- a/Engine/Source/Developer/LauncherServices/Private/Profiles/LauncherProfileManager.cpp
+++ b/Engine/Source/Developer/LauncherServices/Private/Profiles/LauncherProfileManager.cpp
@@ -9,6 +9,7 @@
 #include "Launcher/LauncherProjectPath.h"
 #include "Profiles/LauncherDeviceGroup.h"
 #include "Profiles/LauncherProfile.h"
+#include "GameProjectHelper.h"
 
 LAUNCHERSERVICES_API bool HasPromotedTarget(const TCHAR* BaseDir, const TCHAR* TargetName, const TCHAR* Platform, EBuildConfiguration Configuration, const TCHAR* Architecture)
 {
@@ -455,6 +456,31 @@ void FLauncherProfileManager::SetProjectPath(const FString& InProjectPath)
 	if (ProjectPath != InProjectPath)
 	{
 		ProjectPath = InProjectPath;
+		ExplicitBuildTargets = FGameProjectHelper::GetExplicitBuildTargetsForProject(ProjectPath);
+
+		if (ExplicitBuildTargets.Num() == 0)
+		{
+			SetBuildTarget(FString());
+		}
+		else
+		{
+			// pick a default - otherwise the user will have to select one every time
+			FString DefaultBuildTarget = ExplicitBuildTargets[0];
+
+			// favor a Game target
+			const TArray<FTargetInfo>& Targets = FDesktopPlatformModule::Get()->GetTargetsForProject(ProjectPath);
+			for (const FTargetInfo& Target : Targets)
+			{
+				if (Target.Type == EBuildTargetType::Game && ExplicitBuildTargets.Contains(Target.Name))
+				{
+					DefaultBuildTarget = Target.Name;
+					break;
+				}
+			}
+
+			SetBuildTarget(DefaultBuildTarget);
+		}
+
 		for (ILauncherProfilePtr Profile : AllProfiles)
 		{
 			if (Profile.IsValid())
@@ -462,9 +488,37 @@ void FLauncherProfileManager::SetProjectPath(const FString& InProjectPath)
 				Profile->FallbackProjectUpdated();
 			}
 		}
+		
+		ProjectChangedDelegate.Broadcast();
+	}
+}
+
+FString FLauncherProfileManager::GetBuildTarget() const
+{
+	return BuildTarget;
+}
+
+void FLauncherProfileManager::SetBuildTarget( const FString& InBuildTarget )
+{
+	if (BuildTarget != InBuildTarget)
+	{
+		BuildTarget = InBuildTarget;
+		for (ILauncherProfilePtr Profile : AllProfiles)
+		{
+			if (Profile.IsValid())
+			{
+				Profile->FallbackBuildTargetUpdated();
+			}
+		}
 	}
 }
 
+const TArray<FString>& FLauncherProfileManager::GetAllExplicitBuildTargetNames() const
+{
+	return ExplicitBuildTargets;
+}
+
+
 void FLauncherProfileManager::LoadDeviceGroups( )
 {
 	if (GConfig != nullptr)
diff --git a/Engine/Source/Developer/LauncherServices/Private/Profiles/LauncherProfileManager.h b/Engine/Source/Developer/LauncherServices/Private/Profiles/LauncherProfileManager.h
--- a/Engine/Source/Developer/LauncherServices/Private/Profiles/LauncherProfileManager.h
+++ b/Engine/Source/Developer/LauncherServices/Private/Profiles/LauncherProfileManager.h
@@ -80,6 +80,11 @@ public:
 		return ProfileRemovedDelegate;
 	}
 
+	virtual FOnLauncherProfileManagerProjectChanged& OnProjectChanged( ) override
+	{
+		return ProjectChangedDelegate;
+	}
+
 	virtual void RemoveDeviceGroup( const ILauncherDeviceGroupRef& DeviceGroup ) override;
 
 	virtual void SaveDeviceGroups() override;
@@ -110,6 +115,12 @@ public:
 
 	virtual void SetProjectPath(const FString& InProjectPath) override;
 
+	virtual FString GetBuildTarget() const override;
+
+	virtual void SetBuildTarget( const FString& InBuildTarget ) override;
+
+	virtual const TArray<FString>& GetAllExplicitBuildTargetNames() const override;
+
 	//~ End ILauncherProfileManager Interface
 
 protected:
@@ -171,6 +182,12 @@ private:
 
 	// Holds the currently selected project path
 	FString ProjectPath;
+
+	// Holds the currently selected build target
+	FString BuildTarget;
+
+	// Holds the current explicit build targets
+	TArray<FString> ExplicitBuildTargets;
 	
 	// Holds all registered profile wizards
 	TArray<ILauncherProfileWizardPtr> ProfileWizards;
@@ -188,5 +205,8 @@ private:
 
 	// Holds a delegate to be invoked when a launcher profile was removed.
 	FOnLauncherProfileManagerProfileRemoved ProfileRemovedDelegate;
+	
+	// Holds a delegate to be invoked when the project changes.
+	FOnLauncherProfileManagerProjectChanged ProjectChangedDelegate;	
 
 };
diff --git a/Engine/Source/Developer/LauncherServices/Public/GameProjectHelper.h b/Engine/Source/Developer/LauncherServices/Public/GameProjectHelper.h
--- a/Engine/Source/Developer/LauncherServices/Public/GameProjectHelper.h
+++ b/Engine/Source/Developer/LauncherServices/Public/GameProjectHelper.h
@@ -6,6 +6,7 @@
 #include "Misc/Paths.h"
 #include "HAL/FileManager.h"
 #include "Misc/PackageName.h"
+#include "DesktopPlatformModule.h"
 
 /**
  * Implements a helper class for finding project specific information.
@@ -129,4 +130,52 @@ public:
 
 		return false;
 	}
+
+	/**
+	 * Returns the list of build targets that must be explicitly specified for the given project.
+	 * @param InProjectFilePath - the project to use
+	 * @param InFilterTargetTypes - if specified, the returned list will be restricted to the given target types.
+	 */
+	static TArray<FString> GetExplicitBuildTargetsForProject( const FString& InProjectFilePath, const TSet<EBuildTargetType>* InFilterTargetTypes = nullptr )
+	{
+		TArray<FString> ExplicitBuildTargetNames;
+
+		if (!InProjectFilePath.IsEmpty())
+		{
+			// ignore the filter if it's empty
+			if (InFilterTargetTypes != nullptr && InFilterTargetTypes->Num() == 0)
+			{
+				InFilterTargetTypes = nullptr;
+			}
+
+			// collect the build targets for the given project
+			TMap<EBuildTargetType,TArray<FString>> BuildTargetNamesByType;
+			const TArray<FTargetInfo>& Targets = FDesktopPlatformModule::Get()->GetTargetsForProject(InProjectFilePath);
+			for (const FTargetInfo& Target : Targets)
+			{
+				if (Target.Type != EBuildTargetType::Game && Target.Type != EBuildTargetType::Client && Target.Type != EBuildTargetType::Server)
+				{
+					continue;
+				}
+				
+				if (InFilterTargetTypes != nullptr && !InFilterTargetTypes->Contains(Target.Type) )
+				{
+					continue;
+				}
+
+				BuildTargetNamesByType.FindOrAdd(Target.Type).Add(Target.Name);
+			}
+
+			// create a list of build targets that need to be specified explicitly (because there is more than one of the same type)
+			for (auto& Itr : BuildTargetNamesByType)
+			{
+				if (Itr.Value.Num() > 1)
+				{
+					ExplicitBuildTargetNames.Append(Itr.Value);
+				}
+			}
+		}
+
+		return MoveTemp(ExplicitBuildTargetNames);
+	}
 };
diff --git a/Engine/Source/Developer/LauncherServices/Public/ILauncherProfile.h b/Engine/Source/Developer/LauncherServices/Public/ILauncherProfile.h
--- a/Engine/Source/Developer/LauncherServices/Public/ILauncherProfile.h
+++ b/Engine/Source/Developer/LauncherServices/Public/ILauncherProfile.h
@@ -225,6 +225,15 @@ namespace ELauncherProfileValidationErrors
 		/** Using I/O store container file(s) requires using UnrealPak */
 		IoStoreRequiresPakFiles,
 
+		/** Build Target and Cook Variant mismatch */
+		BuildTargetCookVariantMismatch,
+
+		/** Build Target is required */
+		BuildTargetIsRequired,
+
+		/** FallbackBuild Target is required */
+		FallbackBuildTargetIsRequired,
+
 		Count
 	};
 }
@@ -355,6 +364,9 @@ DECLARE_MULTICAST_DELEGATE_OneParam(FOnLauncherProfileDeployedDeviceGroupChanged
 /** Delegate type for a change in project */
 DECLARE_MULTICAST_DELEGATE(FOnProfileProjectChanged);
 
+/** Delegate type for a change in build target options */
+DECLARE_MULTICAST_DELEGATE(FOnProfileBuildTargetOptionsChanged);
+
 /** Delegate type for detecting if cook is finished
  *	Used when cooking from the editor.  Specific cook task will wait for the cook to be finished by the editor
  */
@@ -512,13 +524,20 @@ public:
 	 */
 	virtual EBuildConfiguration GetBuildConfiguration( ) const = 0;
 
+	/**
+	 * Checks whether the profile specifies a build target.
+	 *
+	 * @return true if the profile specifies a build target.
+	 */
+	virtual bool HasBuildTargetSpecified() const = 0;
+
 	/**
 	 * Gets the build target.
 	 *
 	 * @Return Target name to build/run (it would match a .Target.cs file)
 	 * @see SetBuildTarget
 	 */
-	virtual const FString& GetBuildTarget() const = 0;
+	virtual FString GetBuildTarget() const = 0;
 
 	/**
 	 * Gets the build configuration name of the cooker.
@@ -1008,6 +1027,16 @@ public:
 	 */
 	virtual void SetBuildConfiguration( EBuildConfiguration Configuration ) = 0;
 
+	/**
+	 * Sets whether this profile specfies a build target.
+	 * 
+	 * @param Specified Whether a build target is specified.
+	 */
+	virtual void SetBuildTargetSpecified(bool Specified) = 0;
+
+	/** Notifies the profile that the fallback build target changed. */
+	virtual void FallbackBuildTargetUpdated() = 0;
+
 	/**
 	 * Sets the build target.
 	 *
@@ -1015,6 +1044,7 @@ public:
 	 * @see GetBuildTarget
 	 */
 	virtual void SetBuildTarget( const FString& TargetName ) = 0;
+
 	/**
 	 * Sets the build configuration of the cooker.
 	 *
@@ -1302,6 +1332,13 @@ public:
 	 */
 	virtual FOnProfileProjectChanged& OnProjectChanged() = 0;
 
+	/**
+	 * Access delegate used when build target options change.
+	 * 
+	 * @return The delegate.
+	 */
+	virtual FOnProfileBuildTargetOptionsChanged& OnBuildTargetOptionsChanged() = 0;
+
 	/**
 	 * Sets whether to use I/O store for optimized loading.
 	 * @param bUseIoStore Whether to use I/O store
@@ -1350,6 +1387,17 @@ public:
 	 */
 	virtual bool ShouldUpdateDeviceFlash() const = 0;
 
+public:
+	/**
+	 * Helper function to get all of the build targets available for this profile, based on its current project & cook platforms
+	 */
+	virtual TArray<FString> GetExplicitBuildTargetNames() const = 0;
+
+	/**
+	 * Whether the profile will require an explicit -target=XXX parameter
+	 */
+	virtual bool RequiresExplicitBuildTargetName() const = 0;
+
 public:
 
 	/** Virtual destructor. */
diff --git a/Engine/Source/Developer/LauncherServices/Public/ILauncherProfileManager.h b/Engine/Source/Developer/LauncherServices/Public/ILauncherProfileManager.h
--- a/Engine/Source/Developer/LauncherServices/Public/ILauncherProfileManager.h
+++ b/Engine/Source/Developer/LauncherServices/Public/ILauncherProfileManager.h
@@ -43,6 +43,13 @@ DECLARE_MULTICAST_DELEGATE_OneParam(FOnLauncherProfileManagerProfileAdded, const
  */
 DECLARE_MULTICAST_DELEGATE_OneParam(FOnLauncherProfileManagerProfileRemoved, const ILauncherProfileRef&);
 
+
+/**
+ * Declares a delegate to be invoked when the project is changed.
+ *
+ */
+DECLARE_MULTICAST_DELEGATE(FOnLauncherProfileManagerProjectChanged);
+
 /** Type definition for shared references to instances of ILauncherProfile. */
 typedef TSharedPtr<class ILauncherProfileWizard> ILauncherProfileWizardPtr;
 
@@ -330,6 +337,24 @@ public:
 	 */
 	virtual void SetProjectPath(const FString& InProjectPath) = 0;
 
+public:
+
+	/**
+	 * Gets the default build target to use.
+	 */
+	virtual FString GetBuildTarget() const = 0;
+	
+	/**
+	 * Sets the default buid target to use.
+	 */
+	virtual void SetBuildTarget( const FString& InBuildTarget ) = 0;
+	
+	/**
+	 * Return a list of all build build targets for the current project that need to be specified explicitly via -target=
+	 * If a build target isn't in the list, it does not need to be specified (and BuildTarget can be an empty string)
+	 */
+	virtual const TArray<FString>& GetAllExplicitBuildTargetNames() const = 0;
+
 public:
 
 	/**
@@ -359,6 +384,15 @@ public:
 	 * @return The delegate.
 	 */
 	virtual FOnLauncherProfileManagerProfileRemoved& OnProfileRemoved( ) = 0;
+	
+
+	/**
+	 * Returns a delegate that is invoked when the project is changed.
+	 *
+	 * @return The delegate.
+	 */
+	virtual FOnLauncherProfileManagerProjectChanged& OnProjectChanged( ) = 0;
+	
 
 public:
 
diff --git a/Engine/Source/Developer/ProjectLauncher/Private/Models/ProjectLauncherModel.h b/Engine/Source/Developer/ProjectLauncher/Private/Models/ProjectLauncherModel.h
--- a/Engine/Source/Developer/ProjectLauncher/Private/Models/ProjectLauncherModel.h
+++ b/Engine/Source/Developer/ProjectLauncher/Private/Models/ProjectLauncherModel.h
@@ -133,7 +133,7 @@ public:
 		{
 			if (Profile != SelectedProfile)
 			{
-				TSharedPtr<ILauncherProfile>& PreviousProfile = SelectedProfile;
+				TSharedPtr<ILauncherProfile> PreviousProfile = SelectedProfile;
 				SelectedProfile = Profile;
 
 				ProfileSelectedDelegate.Broadcast(Profile, PreviousProfile);
diff --git a/Engine/Source/Developer/ProjectLauncher/Private/Widgets/Build/SProjectLauncherBuildPage.cpp b/Engine/Source/Developer/ProjectLauncher/Private/Widgets/Build/SProjectLauncherBuildPage.cpp
--- a/Engine/Source/Developer/ProjectLauncher/Private/Widgets/Build/SProjectLauncherBuildPage.cpp
+++ b/Engine/Source/Developer/ProjectLauncher/Private/Widgets/Build/SProjectLauncherBuildPage.cpp
@@ -16,6 +16,7 @@
 #include "Widgets/Layout/SExpandableArea.h"
 #include "Widgets/Shared/SProjectLauncherBuildConfigurationSelector.h"
 #include "Widgets/Shared/SProjectLauncherFormLabel.h"
+#include "Widgets/Shared/SProjectLauncherBuildTargetSelector.h"
 
 
 #define LOCTEXT_NAMESPACE "SProjectLauncherBuildPage"
@@ -117,6 +118,14 @@ void SProjectLauncherBuildPage::Construct(const FArguments& InArgs, const TShare
 									.OnConfigurationSelected(this, &SProjectLauncherBuildPage::HandleBuildConfigurationSelectorConfigurationSelected)
 									.Text(this, &SProjectLauncherBuildPage::HandleBuildConfigurationSelectorText)
 							]
+
+						+ SHorizontalBox::Slot()
+							.AutoWidth()
+							[
+								SNew(SProjectLauncherBuildTargetSelector, Model.ToSharedRef())
+									.UseProfile(true)
+							]
+
 					]
 			]
 
diff --git a/Engine/Source/Developer/ProjectLauncher/Private/Widgets/Deploy/SProjectLauncherSimpleDeviceListRow.h b/Engine/Source/Developer/ProjectLauncher/Private/Widgets/Deploy/SProjectLauncherSimpleDeviceListRow.h
--- a/Engine/Source/Developer/ProjectLauncher/Private/Widgets/Deploy/SProjectLauncherSimpleDeviceListRow.h
+++ b/Engine/Source/Developer/ProjectLauncher/Private/Widgets/Deploy/SProjectLauncherSimpleDeviceListRow.h
@@ -88,6 +88,8 @@ public:
 		SimpleProfile = Model->GetProfileManager()->FindOrAddSimpleProfile(DeviceProxy->GetName());
 
 		LaunchProfile = Model->GetProfileManager()->CreateUnsavedProfile(DeviceProxy->GetName());
+		LaunchProfile->SetProjectSpecified(false); // device launch profiles should always use the fallback project & build target
+		LaunchProfile->SetBuildTargetSpecified(false);
 		UpdateProfile();
 
 		TSharedRef<SUniformGridPanel> NameGrid = SNew(SUniformGridPanel).SlotPadding(FMargin(0.0f, 1.0f));
diff --git a/Engine/Source/Developer/ProjectLauncher/Private/Widgets/SProjectLauncher.cpp b/Engine/Source/Developer/ProjectLauncher/Private/Widgets/SProjectLauncher.cpp
--- a/Engine/Source/Developer/ProjectLauncher/Private/Widgets/SProjectLauncher.cpp
+++ b/Engine/Source/Developer/ProjectLauncher/Private/Widgets/SProjectLauncher.cpp
@@ -21,6 +21,7 @@
 #include "Widgets/Profile/SProjectLauncherProfileListView.h"
 #include "Widgets/Progress/SProjectLauncherProgress.h"
 #include "Widgets/Project/SProjectLauncherProjectPicker.h"
+#include "Widgets/Shared/SProjectLauncherBuildTargetSelector.h"
 #include "Widgets/Settings/SProjectLauncherSettings.h"
 #include "Styling/AppStyle.h"
 #include "Styling/StyleColors.h"
@@ -117,6 +118,15 @@ void SProjectLauncher::Construct(const FArguments& InArgs, const TSharedRef<SDoc
 							SNew(SProjectLauncherProjectPicker, InModel)
 						]
 
+						// Build Target
+						+ SHorizontalBox::Slot()
+						.AutoWidth()
+						.Padding(0,4)
+						[
+							SNew(SProjectLauncherBuildTargetSelector, InModel)
+							.UseProfile(false)
+						]
+
 						// Advanced Button
 						+ SHorizontalBox::Slot()
 						.AutoWidth()
diff --git a/Engine/Source/Developer/ProjectLauncher/Private/Widgets/Shared/SProjectLauncherBuildTargetSelector.cpp b/Engine/Source/Developer/ProjectLauncher/Private/Widgets/Shared/SProjectLauncherBuildTargetSelector.cpp
new file mode 100644
--- /dev/null
+++ b/Engine/Source/Developer/ProjectLauncher/Private/Widgets/Shared/SProjectLauncherBuildTargetSelector.cpp
@@ -0,0 +1,253 @@
+// Copyright Epic Games, Inc. All Rights Reserved.
+
+#include "SProjectLauncherBuildTargetSelector.h"
+
+#include "IDesktopPlatform.h"
+#include "DesktopPlatformModule.h"
+#include "Styling/AppStyle.h"
+#include "Textures/SlateIcon.h"
+#include "Widgets/SBoxPanel.h"
+#include "Widgets/Layout/SBorder.h"
+#include "Widgets/Text/STextBlock.h"
+#include "Widgets/Shared/SProjectLauncherFormLabel.h"
+
+
+
+#define LOCTEXT_NAMESPACE "SProjectLauncherUnrealBuildTargetSelector"
+
+
+/* SProjectLauncherBuildTargetSelector structors
+ *****************************************************************************/
+
+SProjectLauncherBuildTargetSelector::~SProjectLauncherBuildTargetSelector()
+{
+	if (Model.IsValid())
+	{
+		Model->OnProfileSelected().RemoveAll(this);
+		Model->GetProfileManager()->OnProjectChanged().RemoveAll(this);
+	}
+
+	if (bUseProfile)
+	{
+		ILauncherProfilePtr Profile = GetProfile();
+		if (Profile.IsValid())
+		{
+			Profile->OnProjectChanged().RemoveAll(this);
+			Profile->OnBuildTargetOptionsChanged().RemoveAll(this);
+		}
+	}
+}
+
+
+/* SProjectLauncherBuildTargetSelector interface
+ *****************************************************************************/
+
+void SProjectLauncherBuildTargetSelector::Construct(const FArguments& InArgs, const TSharedRef<FProjectLauncherModel>& InModel)
+{
+	Model = InModel;
+	bUseProfile = InArgs._UseProfile.Get();
+
+	if (bUseProfile)
+	{
+		Model->OnProfileSelected().AddSP(this, &SProjectLauncherBuildTargetSelector::HandleProfileManagerProfileSelected);
+	}
+	Model->GetProfileManager()->OnProjectChanged().AddSP(this, &SProjectLauncherBuildTargetSelector::HandleProfileProjectChanged);
+	
+	FText LabelText = bUseProfile ? LOCTEXT("BuildTargetComboBoxLabel", "Build Target") : LOCTEXT("DefaultBuildTargetComboBoxLabel", "Default Build Target");
+
+	ChildSlot
+	[
+		SNew(SBorder)
+		.BorderImage(FAppStyle::GetBrush("ToolPanel.GroupBorder"))
+		.Padding(0)
+		.Visibility(this, &SProjectLauncherBuildTargetSelector::IsBuildTargetVisible)
+		[
+			SNew(SHorizontalBox)
+
+			+ SHorizontalBox::Slot()
+				.Padding(16, 0, 12, 0)
+				.AutoWidth()
+				[
+					SNew(SProjectLauncherFormLabel)
+					.LabelText(LabelText)
+				]
+
+			+ SHorizontalBox::Slot()
+				.AutoWidth()
+				[
+				SAssignNew(BuildTargetCombo,SComboBox<TSharedPtr<FString>>)
+				.OptionsSource(&BuildTargetList)
+				.OnSelectionChanged( this, &SProjectLauncherBuildTargetSelector::HandleBuildTargetSelectionChanged )
+				.OnGenerateWidget( this, &SProjectLauncherBuildTargetSelector::HandleBuildTargetGenerateWidget )
+				.ContentPadding(FMargin(4.0f, 0.0f, 4.0f, 0.0f))
+				.Content()
+				[
+					SNew(STextBlock)
+					.Text(this, &SProjectLauncherBuildTargetSelector::HandleBuildTargetComboButtonText)
+				]
+			]
+		]
+	];
+
+	RefreshBuildTargetList();
+}
+
+
+
+
+FText SProjectLauncherBuildTargetSelector::GetBuildTargetText( FString BuildTarget ) const
+{
+	if (!BuildTarget.IsEmpty())
+	{
+		return FText::FromString(BuildTarget);
+	}
+
+	return LOCTEXT("NoBuildTargetText", "Default");
+}
+
+FText SProjectLauncherBuildTargetSelector::HandleBuildTargetComboButtonText() const
+{
+	FString BuildTarget;
+
+	if (bUseProfile)
+	{
+		ILauncherProfilePtr Profile = GetProfile();
+		if (Profile.IsValid() && Profile->HasBuildTargetSpecified())
+		{
+			BuildTarget = Profile->GetBuildTarget();
+		}
+	}
+	else
+	{
+		BuildTarget = Model->GetProfileManager()->GetBuildTarget();
+	}
+
+	return GetBuildTargetText(BuildTarget);
+}
+
+
+void SProjectLauncherBuildTargetSelector::HandleBuildTargetSelectionChanged(TSharedPtr<FString> Item, ESelectInfo::Type SelectInfo)
+{
+	if (Item.IsValid())
+	{
+		SetBuildTarget(*Item.Get());
+	}
+}
+
+TSharedRef<SWidget> SProjectLauncherBuildTargetSelector::HandleBuildTargetGenerateWidget(TSharedPtr<FString> Item)
+{
+	FText BuildTargetText = GetBuildTargetText(*Item.Get());
+	return SNew(STextBlock).Text(BuildTargetText);
+}
+
+EVisibility SProjectLauncherBuildTargetSelector::IsBuildTargetVisible() const
+{
+	if (bUseProfile)
+	{
+		ILauncherProfilePtr Profile = GetProfile();
+		if (Profile.IsValid() && Profile->RequiresExplicitBuildTargetName())
+		{
+			return EVisibility::Visible;
+		}
+	}
+	else
+	{
+		if (Model->GetProfileManager()->GetAllExplicitBuildTargetNames().Num() > 0)
+		{
+			return EVisibility::Visible;
+		}
+	}
+
+	return EVisibility::Collapsed;
+}
+
+
+
+void SProjectLauncherBuildTargetSelector::SetBuildTarget(const FString& BuildTarget)
+{
+	if (bUseProfile)
+	{
+		ILauncherProfilePtr Profile = GetProfile();
+		if (Profile.IsValid())
+		{
+			Profile->SetBuildTarget(BuildTarget);
+			Profile->SetBuildTargetSpecified(true);
+		}
+	}
+	else
+	{
+		Model->GetProfileManager()->SetBuildTarget(BuildTarget);
+	}
+}
+
+void SProjectLauncherBuildTargetSelector::HandleProfileManagerProfileSelected(const ILauncherProfilePtr& SelectedProfile, const ILauncherProfilePtr& PreviousProfile)
+{
+	if (PreviousProfile.IsValid())
+	{
+		PreviousProfile->OnProjectChanged().RemoveAll(this);
+		PreviousProfile->OnBuildTargetOptionsChanged().RemoveAll(this);
+	}
+	if (SelectedProfile.IsValid())
+	{
+		SelectedProfile->OnProjectChanged().AddSP(this, &SProjectLauncherBuildTargetSelector::HandleProfileProjectChanged);
+		SelectedProfile->OnBuildTargetOptionsChanged().AddSP(this, &SProjectLauncherBuildTargetSelector::HandleProfileBuildTargetOptionsChanged);
+	}
+	RefreshBuildTargetList();
+}
+
+
+void SProjectLauncherBuildTargetSelector::HandleProfileProjectChanged()
+{
+	RefreshBuildTargetList();
+}
+
+void SProjectLauncherBuildTargetSelector::HandleProfileBuildTargetOptionsChanged()
+{
+	RefreshBuildTargetList();
+}
+
+
+void SProjectLauncherBuildTargetSelector::RefreshBuildTargetList()
+{
+	// make a list of all suitable build target names
+	TArray<FString> ExplicitBuildTargetNames;
+	if (bUseProfile)
+	{
+		ILauncherProfilePtr Profile = GetProfile();
+		if (Profile.IsValid())
+		{
+			ExplicitBuildTargetNames = Profile->GetExplicitBuildTargetNames();
+		}
+	}
+	else
+	{
+		ExplicitBuildTargetNames = Model->GetProfileManager()->GetAllExplicitBuildTargetNames();
+	}
+
+	BuildTargetList.Reset();
+	for ( const FString& BuildTarget : ExplicitBuildTargetNames)
+	{
+		BuildTargetList.Add(MakeShared<FString>(BuildTarget));
+	}
+
+	// add a blank entry to allow resetting back to 'default'
+	BuildTargetList.Add(MakeShared<FString>());
+
+	// refresh UI
+	if (BuildTargetCombo.IsValid())
+	{
+		BuildTargetCombo->RefreshOptions();
+	}
+}
+
+ILauncherProfilePtr SProjectLauncherBuildTargetSelector::GetProfile() const
+{
+	if (ensure(bUseProfile) && Model.IsValid())
+	{
+		return Model->GetSelectedProfile();
+	}
+
+	return nullptr;
+}
+
+#undef LOCTEXT_NAMESPACE
diff --git a/Engine/Source/Developer/ProjectLauncher/Private/Widgets/Shared/SProjectLauncherBuildTargetSelector.h b/Engine/Source/Developer/ProjectLauncher/Private/Widgets/Shared/SProjectLauncherBuildTargetSelector.h
new file mode 100644
--- /dev/null
+++ b/Engine/Source/Developer/ProjectLauncher/Private/Widgets/Shared/SProjectLauncherBuildTargetSelector.h
@@ -0,0 +1,72 @@
+// Copyright Epic Games, Inc. All Rights Reserved.
+
+#pragma once
+
+#include "CoreMinimal.h"
+#include "Misc/Attribute.h"
+#include "Layout/Visibility.h"
+#include "Widgets/DeclarativeSyntaxSupport.h"
+#include "Widgets/SWidget.h"
+#include "Widgets/SCompoundWidget.h"
+#include "Widgets/Input/SComboBox.h"
+#include "Models/ProjectLauncherModel.h"
+
+class Error;
+
+DECLARE_DELEGATE_OneParam(FOnLauncherProfileChanged, ILauncherProfilePtr);
+
+class SProjectLauncherBuildTargetSelector
+	: public SCompoundWidget
+{
+public:
+
+	SLATE_BEGIN_ARGS(SProjectLauncherBuildTargetSelector)
+		: _UseProfile(false)
+		{}
+		SLATE_ATTRIBUTE(bool, UseProfile) // whether we are showing items for the current profile or not
+	SLATE_END_ARGS()
+
+public:
+
+	/**
+	 * Destructor.
+	 */
+	~SProjectLauncherBuildTargetSelector();
+
+public:
+
+	/**
+	 * Constructs the widget.
+	 *
+	 * @param InArgs The Slate argument list.
+	 * @param InModel The data model.
+	 */
+	void Construct(const FArguments& InArgs, const TSharedRef<FProjectLauncherModel>& InModel);
+
+
+private:
+
+	void RefreshBuildTargetList();
+
+	EVisibility IsBuildTargetVisible() const;
+	FText HandleBuildTargetComboButtonText() const;
+	FText GetBuildTargetText( FString BuildTarget ) const;
+	void HandleBuildTargetSelectionChanged(TSharedPtr<FString> Item, ESelectInfo::Type SelectInfo);
+	TSharedRef<SWidget> HandleBuildTargetGenerateWidget(TSharedPtr<FString> Item);
+	void HandleProfileManagerProfileSelected(const ILauncherProfilePtr& SelectedProfile, const ILauncherProfilePtr& PreviousProfile);
+	void HandleProfileProjectChanged();
+	void HandleProfileBuildTargetOptionsChanged();
+
+	void SetBuildTarget(const FString& BuildTarget);
+
+	ILauncherProfilePtr GetProfile() const;
+
+
+private:
+
+	TSharedPtr<SComboBox<TSharedPtr<FString>>> BuildTargetCombo;
+	TArray<TSharedPtr<FString>> BuildTargetList;
+	bool bUseProfile;
+	TSharedPtr<FProjectLauncherModel> Model;
+
+};
diff --git a/Engine/Source/Developer/ProjectLauncher/Private/Widgets/Shared/SProjectLauncherValidation.h b/Engine/Source/Developer/ProjectLauncher/Private/Widgets/Shared/SProjectLauncherValidation.h
--- a/Engine/Source/Developer/ProjectLauncher/Private/Widgets/Shared/SProjectLauncherValidation.h
+++ b/Engine/Source/Developer/ProjectLauncher/Private/Widgets/Shared/SProjectLauncherValidation.h
@@ -176,6 +176,21 @@ public:
 				MakeValidationMessage(TEXT("Icons.Error"), LOCTEXT("IoStoreRequiresPakFilesError", "UnrealPak must be selected when using I/O store container file(s)").ToString(), ELauncherProfileValidationErrors::IoStoreRequiresPakFiles)
 			];
 
+		VertBox->AddSlot().AutoHeight()
+			[
+				MakeValidationMessage(TEXT("Icons.Error"), LOCTEXT("BuildTargetCookVariantMismatch", "Build Target and Cook Variant mismatch.").ToString(), ELauncherProfileValidationErrors::BuildTargetCookVariantMismatch)
+			];
+
+		VertBox->AddSlot().AutoHeight()
+			[
+				MakeValidationMessage(TEXT("Icons.Error"), LOCTEXT("BuildTargetIsRequired", "This profile requires an explicit Build Target set.").ToString(), ELauncherProfileValidationErrors::BuildTargetIsRequired)
+			];
+
+		VertBox->AddSlot().AutoHeight()
+			[
+				MakeValidationMessage(TEXT("Icons.Error"), LOCTEXT("FallbackBuildTargetIsRequired", "An explicit Default Build Target is required for the selected Variant.").ToString(), ELauncherProfileValidationErrors::FallbackBuildTargetIsRequired)
+			];
+
 		check(VertBox->NumSlots() == ELauncherProfileValidationErrors::Count);
 
 		ChildSlot[VertBox.ToSharedRef()];
