Fixed forwarding of -DDC=<Graph> to avoid initializing DDC and only forward non-default graph names

diff --git a/Engine/Source/Developer/DerivedDataCache/Private/DerivedDataCacheModule.cpp b/Engine/Source/Developer/DerivedDataCache/Private/DerivedDataCacheModule.cpp
--- a/Engine/Source/Developer/DerivedDataCache/Private/DerivedDataCacheModule.cpp
+++ b/Engine/Source/Developer/DerivedDataCache/Private/DerivedDataCacheModule.cpp
@@ -35,6 +35,11 @@ public:
 		return &GDerivedDataLegacyCache;
 	}
 
+	FDerivedDataCacheInterface* const* GetCache() final
+	{
+		return &GDerivedDataLegacyCache;
+	}
+
 	void CreateCacheOnce()
 	{
 		FScopeLock Lock(&CreateLock);
@@ -111,6 +116,11 @@ ICache& GetCache()
 	return *Cache;
 }
 
+ICache* TryGetCache()
+{
+	return Private::GDerivedDataCache;
+}
+
 IBuild& GetBuild()
 {
 	if (IBuild* Build = Private::GDerivedDataBuild)
diff --git a/Engine/Source/Developer/DerivedDataCache/Public/DerivedDataCache.h b/Engine/Source/Developer/DerivedDataCache/Public/DerivedDataCache.h
--- a/Engine/Source/Developer/DerivedDataCache/Public/DerivedDataCache.h
+++ b/Engine/Source/Developer/DerivedDataCache/Public/DerivedDataCache.h
@@ -373,6 +373,9 @@ UE_API bool LoadFromCompactBinary(FCbFieldView Field, FCacheGetChunkRequest& Out
 /** Returns a reference to the cache. Asserts if not available. */
 UE_API ICache& GetCache();
 
+/** Returns a pointer to the cache. Null if not available or not created. */
+UE_API ICache* TryGetCache();
+
 ///////////////////////////////////////////////////////////////////////////////////////////////////////////////////////
 
 } // UE::DerivedData
diff --git a/Engine/Source/Developer/DerivedDataCache/Public/DerivedDataCacheModule.h b/Engine/Source/Developer/DerivedDataCache/Public/DerivedDataCacheModule.h
--- a/Engine/Source/Developer/DerivedDataCache/Public/DerivedDataCacheModule.h
+++ b/Engine/Source/Developer/DerivedDataCache/Public/DerivedDataCacheModule.h
@@ -25,4 +25,13 @@ public:
 	 * the destruction of the cache without polling this function or monitoring the module lifetime.
 	 */
 	virtual FDerivedDataCacheInterface* const* CreateOrGetCache() = 0;
+
+	/**
+	 * Returns a pointer to the cache, which may be null.
+	 *
+	 * This returns a non-null pointer to a cache which may be null and becomes null when the module
+	 * shuts down and destroys the cache. This extra level of indirection allows a caller to observe
+	 * the destruction of the cache without polling this function or monitoring the module lifetime.
+	 */
+	virtual FDerivedDataCacheInterface* const* GetCache() = 0;
 };
diff --git a/Engine/Source/Developer/LauncherServices/Private/Launcher/LauncherWorker.cpp b/Engine/Source/Developer/LauncherServices/Private/Launcher/LauncherWorker.cpp
--- a/Engine/Source/Developer/LauncherServices/Private/Launcher/LauncherWorker.cpp
+++ b/Engine/Source/Developer/LauncherServices/Private/Launcher/LauncherWorker.cpp
@@ -642,9 +642,13 @@ FString FLauncherWorker::CreateUATCommand( const ILauncherProfileRef& InProfile,
 				UATCommand += TEXT(" -zenstore -fileserver -skipserver");
 			}
 
-			if (FDerivedDataCacheInterface* DDC = GetDerivedDataCache())
+			if (FDerivedDataCacheInterface* DDC = TryGetDerivedDataCache())
 			{
-				UATCommand += FString::Printf(TEXT(" -ddc=%s"), DDC->GetGraphName());
+				const TCHAR* GraphName = DDC->GetGraphName();
+				if (FCString::Strcmp(GraphName, DDC->GetDefaultGraphName()))
+				{
+					UATCommand += FString::Printf(TEXT(" -DDC=%s"), GraphName);
+				}
 			}
 
 			if (InProfile->IsPackingWithUnrealPak())
diff --git a/Engine/Source/Editor/TurnkeySupport/Private/TurnkeySupportModule.cpp b/Engine/Source/Editor/TurnkeySupport/Private/TurnkeySupportModule.cpp
--- a/Engine/Source/Editor/TurnkeySupport/Private/TurnkeySupportModule.cpp
+++ b/Engine/Source/Editor/TurnkeySupport/Private/TurnkeySupportModule.cpp
@@ -362,9 +362,13 @@ public:
 		{
 			BuildCookRunParams += TEXT(" -SkipCookingEditorContent");
 		}
-		if (FDerivedDataCacheInterface* DDC = GetDerivedDataCache())
+		if (FDerivedDataCacheInterface* DDC = TryGetDerivedDataCache())
 		{
-			BuildCookRunParams += FString::Printf(TEXT(" -ddc=%s"), DDC->GetGraphName());
+			const TCHAR* GraphName = DDC->GetGraphName();
+			if (FCString::Strcmp(GraphName, DDC->GetDefaultGraphName()))
+			{
+				BuildCookRunParams += FString::Printf(TEXT(" -DDC=%s"), DDC->GetGraphName());
+			}
 		}
 		if (FApp::IsEngineInstalled())
 		{
diff --git a/Engine/Source/Runtime/Core/Private/Misc/CoreMisc.cpp b/Engine/Source/Runtime/Core/Private/Misc/CoreMisc.cpp
--- a/Engine/Source/Runtime/Core/Private/Misc/CoreMisc.cpp
+++ b/Engine/Source/Runtime/Core/Private/Misc/CoreMisc.cpp
@@ -120,6 +120,18 @@ FDerivedDataCacheInterface& GetDerivedDataCacheRef()
 	return *Cache;
 }
 
+FDerivedDataCacheInterface* TryGetDerivedDataCache()
+{
+	if (IDerivedDataCacheModule* Module = FModuleManager::GetModulePtr<IDerivedDataCacheModule>("DerivedDataCache"))
+	{
+		if (FDerivedDataCacheInterface* const* Cache = Module->GetCache())
+		{
+			return *Cache;
+		}
+	}
+	return nullptr;
+}
+
 class ITargetPlatformManagerModule* GetTargetPlatformManager(bool bFailOnInitErrors)
 {
 	static class ITargetPlatformManagerModule* SingletonInterface = NULL;
diff --git a/Engine/Source/Runtime/Core/Public/Misc/CoreMisc.h b/Engine/Source/Runtime/Core/Public/Misc/CoreMisc.h
--- a/Engine/Source/Runtime/Core/Public/Misc/CoreMisc.h
+++ b/Engine/Source/Runtime/Core/Public/Misc/CoreMisc.h
@@ -78,6 +78,9 @@ CORE_API class FDerivedDataCacheInterface* GetDerivedDataCache();
 /** Returns the derived data cache interface, or fatal error if it is not available. */
 CORE_API class FDerivedDataCacheInterface& GetDerivedDataCacheRef();
 
+/** Returns the derived data cache interface if it is available and initialized, otherwise null. */
+CORE_API class FDerivedDataCacheInterface* TryGetDerivedDataCache();
+
 /**
  * Return the Target Platform Manager interface, if it is available, otherwise return nullptr.
  *
