Pass along the RHI specific arguments to launched processes.

diff --git a/Engine/Source/Developer/LauncherServices/Private/Launcher/LauncherWorker.cpp b/Engine/Source/Developer/LauncherServices/Private/Launcher/LauncherWorker.cpp
--- a/Engine/Source/Developer/LauncherServices/Private/Launcher/LauncherWorker.cpp
+++ b/Engine/Source/Developer/LauncherServices/Private/Launcher/LauncherWorker.cpp
@@ -282,6 +282,31 @@ static void AddDeviceToLaunchCommand(const FString& DeviceId, TSharedPtr<ITarget
 		RoleCommands += TEXT(" -opengl");
 	}
 
+	if (FParse::Param(FCommandLine::Get(), TEXT("d3d11")) || FParse::Param(FCommandLine::Get(), TEXT("dx11")))
+	{
+		RoleCommands += TEXT(" -d3d11");
+	}
+
+	if (FParse::Param(FCommandLine::Get(), TEXT("d3d12")) || FParse::Param(FCommandLine::Get(), TEXT("dx12")))
+	{
+		RoleCommands += TEXT(" -d3d12");
+	}
+
+	if (FParse::Param(FCommandLine::Get(), TEXT("es31")) || FParse::Param(FCommandLine::Get(), TEXT("FeatureLevelES31")) || FParse::Param(FCommandLine::Get(), TEXT("FeatureLevelES3_1")))
+	{
+		RoleCommands += TEXT(" -es31");
+	}
+
+	if (FParse::Param(FCommandLine::Get(), TEXT("sm5")))
+	{
+		RoleCommands += TEXT(" -sm5");
+	}
+
+	if (FParse::Param(FCommandLine::Get(), TEXT("sm6")))
+	{
+		RoleCommands += TEXT(" -sm6");
+	}
+
 	if (FParse::Param(FCommandLine::Get(), TEXT("vulkan")))
 	{
 		FName Variant = DeviceProxy->GetTargetDeviceVariant(DeviceId);
