fix for the editor always trying to build even with promoted binaries.

#rnx

diff --git a/Engine/Source/Developer/LauncherServices/Private/Profiles/LauncherProfile.h b/Engine/Source/Developer/LauncherServices/Private/Profiles/LauncherProfile.h
--- a/Engine/Source/Developer/LauncherServices/Private/Profiles/LauncherProfile.h
+++ b/Engine/Source/Developer/LauncherServices/Private/Profiles/LauncherProfile.h
@@ -746,6 +746,8 @@ public:
 		{
 			if (FApp::GetEngineIsPromotedBuild())
 			{
+				bBuild = false;
+
 				TArray<FString> TargetPlatformNames = FindPlatforms();
 				for (const FString& TargetPlatformName : TargetPlatformNames)
 				{
@@ -777,6 +779,7 @@ public:
 					else
 					{
 						UE_LOG(LogLauncherProfile, Log, TEXT("Unable to find any targets for platform %s - forcing build"), *TargetPlatformName);
+						bBuild = true;
 						break;
 					}
 
@@ -784,6 +787,7 @@ public:
 					FString BuildPlatform = PlatformInfo.DataDrivenPlatformInfo->UBTPlatformString;
 					if (!HasPromotedTarget(*ReceiptDir, *TargetName, *BuildPlatform, BuildConfiguration, nullptr))
 					{
+						bBuild = true;
 						break;
 					}
 				}
