Project Launcher Build Target fixes:
 - saved profiles default to not having the build target specified as it wasn't originally serialized.
 - using the same validation logic regardless of whether the project is specified or not to keep things simpler.
 - editor & server do not need an explicit build target when there's multiple target.cs files because the default can be specified in the engine ini.

#rnx

diff --git a/Engine/Source/Developer/LauncherServices/Private/Profiles/LauncherProfile.h b/Engine/Source/Developer/LauncherServices/Private/Profiles/LauncherProfile.h
--- a/Engine/Source/Developer/LauncherServices/Private/Profiles/LauncherProfile.h
+++ b/Engine/Source/Developer/LauncherServices/Private/Profiles/LauncherProfile.h
@@ -2047,7 +2047,7 @@ public:
 		// default build settings
 		BuildMode = ELauncherProfileBuildModes::Auto;
 		BuildUAT = !FApp::GetEngineIsPromotedBuild() && !FApp::IsEngineInstalled();
-		BuildTargetSpecified = true;
+		BuildTargetSpecified = false;
 
 		// default cook settings
 		CookConfiguration = FApp::GetBuildConfiguration();
@@ -2925,63 +2925,47 @@ protected:
 	void ValidateBuildTarget()
 	{
 		bool bBuildTargetIsRequired = false;
-		bool bBuildTargetCookVariantMismatch = false;
+		bool bBuildTargetIsSelected = false;
 
 		FString BuildTarget = GetBuildTarget();
 		TSet<EBuildTargetType> CookTargetTypes = GetCookTargetTypes();
+		if (CookTargetTypes.Num() == 0)
+		{
+			CookTargetTypes.Add(EBuildTargetType::Game); // UAT defaults to 'Game' too
+		}
 
-		if (HasProjectSpecified())
+		// check all build targets
+		const TArray<FString>& BuildTargetNames = HasProjectSpecified() ? ExplictBuildTargetNames : LauncherProfileManager->GetAllExplicitBuildTargetNames();	
+		const TArray<FTargetInfo>& Targets = FDesktopPlatformModule::Get()->GetTargetsForProject(GetProjectPath());
+		for (const FTargetInfo& Target : Targets)
 		{
-			if (RequiresExplicitBuildTargetName())
+			if (CookTargetTypes.Contains(Target.Type))
 			{
-				if (CookTargetTypes.Num() > 1 || (!BuildTarget.IsEmpty() && !ExplictBuildTargetNames.Contains(BuildTarget) ) )
-				{
-					// can only build the same Variant (Game, Client, etc) as the selected build target
-					bBuildTargetCookVariantMismatch = true;
-				}
-				else if (BuildTarget.IsEmpty())
+				if (BuildTarget.IsEmpty())
 				{
-					// multiple .target.cs files defined of the same Variant - need to specify one
-					bBuildTargetIsRequired = true;
-				}
-			}
-		}		
-		else
-		{
-			// this profile is using the fallback project. Need to check all build targets instead
-			bool bBuildTargetIsCookable = false;
-			bool bVariantRequiresBuildTarget = false;
+					// editor & server can have a default build target specified in engine ini so no need to enforce this (UAT will give informative error if the ini isn't set up)
+					bool bSupportsDefaultBuildTarget = (Target.Type == EBuildTargetType::Editor || Target.Type == EBuildTargetType::Server); 
 
-			const TArray<FTargetInfo>& Targets = FDesktopPlatformModule::Get()->GetTargetsForProject(GetProjectPath());
-			for (const FTargetInfo& Target : Targets)
-			{
-				if (CookTargetTypes.IsEmpty() || CookTargetTypes.Contains(Target.Type))
-				{
-					if (BuildTarget.IsEmpty())
+					if (!bSupportsDefaultBuildTarget && BuildTargetNames.Contains(Target.Name))
 					{
-						if (LauncherProfileManager->GetAllExplicitBuildTargetNames().Contains(Target.Name))
-						{
-							// a currently selected Variant has multiple .target.cs files - need to specify one
-							bBuildTargetIsRequired = true;
-							break;
-						}
+						bBuildTargetIsRequired = true;
+						break;
 					}
-					else
+				}
+				else
+				{
+					if (Target.Name == BuildTarget)
 					{
-						if (Target.Name == BuildTarget)
-						{
-							bBuildTargetIsCookable = true;
-							break;
-						}
-					}				
+						bBuildTargetIsSelected = true;
+						break;
+					}
 				}
 			}
+		}
 
-			if (!BuildTarget.IsEmpty() && !bBuildTargetIsCookable)
-			{
-				// currently build target does not reference a Variant with multiple .target.cs files
-				bBuildTargetCookVariantMismatch = true;
-			}
+		if (!BuildTarget.IsEmpty() && !bBuildTargetIsSelected)
+		{
+			ValidationErrors.Add(ELauncherProfileValidationErrors::BuildTargetCookVariantMismatch);
 		}
 
 		if (bBuildTargetIsRequired)
@@ -2995,12 +2979,6 @@ protected:
 				ValidationErrors.Add(ELauncherProfileValidationErrors::FallbackBuildTargetIsRequired);
 			}
 		}
-
-		if (bBuildTargetCookVariantMismatch)
-		{
-			ValidationErrors.Add(ELauncherProfileValidationErrors::BuildTargetCookVariantMismatch);
-		}
-
 	}
 
 
