Strip -noclient when cooking server targets along with any other build.

#rnx

diff --git a/Engine/Source/Developer/LauncherServices/Private/Launcher/LauncherWorker.cpp b/Engine/Source/Developer/LauncherServices/Private/Launcher/LauncherWorker.cpp
--- a/Engine/Source/Developer/LauncherServices/Private/Launcher/LauncherWorker.cpp
+++ b/Engine/Source/Developer/LauncherServices/Private/Launcher/LauncherWorker.cpp
@@ -455,8 +455,8 @@ FString FLauncherWorker::CreateUATCommand( const ILauncherProfileRef& InProfile,
 		}
 	}
 
-	// If both Client and Server are desired to be built avoid Server causing clients to not be built PlatformInfo wise
-	if (OptionalParams.Contains(TEXT("-client")) && OptionalParams.Contains(TEXT("-noclient")))
+	// If both Client/Game and Server are desired to be built avoid Server causing clients/game to not be built PlatformInfo wise
+	if (ServerPlatforms.Len() > 0 && Platforms.Len() > 0 && OptionalParams.Contains(TEXT("-noclient")))
 	{
 		OptionalParams = OptionalParams.Replace(TEXT("-noclient"), TEXT(""));
 	}
