LauncherWorker: Fix leading plus added to "-map=+Map1+Map2" argument when CookedMaps is not empty and InitialMap is empty.
#rnx

diff --git a/Engine/Source/Developer/LauncherServices/Private/Launcher/LauncherWorker.cpp b/Engine/Source/Developer/LauncherServices/Private/Launcher/LauncherWorker.cpp
--- a/Engine/Source/Developer/LauncherServices/Private/Launcher/LauncherWorker.cpp
+++ b/Engine/Source/Developer/LauncherServices/Private/Launcher/LauncherWorker.cpp
@@ -576,17 +576,24 @@ FString FLauncherWorker::CreateUATCommand( const ILauncherProfileRef& InProfile,
 		*InProfile->GetAdditionalCommandLineParameters());
 
 	// map list
-	FString MapList = TEXT(" -map=");
+	FString MapList;
 	const TArray<FString>& CookedMaps = InProfile->GetCookedMaps();
 	if (CookedMaps.Num() > 0 && (InProfile->GetCookMode() == ELauncherProfileCookModes::ByTheBook || InProfile->GetCookMode() == ELauncherProfileCookModes::ByTheBookInEditor))
 	{
-		MapList += InitialMap;
-		MapList += TEXT("+");
+		if (!InitialMap.IsEmpty())
+		{
+			MapList += InitialMap;
+			MapList += TEXT("+");
+		}
 		MapList += FString::Join(CookedMaps, TEXT("+"));
 	}
 	else
 	{
-		MapList += InitialMap;
+		MapList = InitialMap;
+	}
+	if (!MapList.IsEmpty())
+	{
+		MapList = FString(TEXT(" -map=")) + MapList;
 	}
 
 	// culture list
