Fixes for multiple build targets:

 - do not need to specify a default server target when doing -cook
 - UFE allows launching default Game and Client targets now that a default can be specified in BuildSettings (note that it won't verify that a default has been configured - UAT handles this)

diff --git a/Engine/Source/Developer/LauncherServices/Private/Profiles/LauncherProfile.h b/Engine/Source/Developer/LauncherServices/Private/Profiles/LauncherProfile.h
--- a/Engine/Source/Developer/LauncherServices/Private/Profiles/LauncherProfile.h
+++ b/Engine/Source/Developer/LauncherServices/Private/Profiles/LauncherProfile.h
@@ -2949,8 +2949,8 @@ protected:
 			{
 				if (BuildTarget.IsEmpty())
 				{
-					// editor & server can have a default build target specified in engine ini so no need to enforce this (UAT will give informative error if the ini isn't set up)
-					bool bSupportsDefaultBuildTarget = (Target.Type == EBuildTargetType::Editor || Target.Type == EBuildTargetType::Server); 
+					// most target types can have a default build target specified in engine ini so no need to enforce this (UAT will give informative error if the ini isn't set up)
+					bool bSupportsDefaultBuildTarget = (Target.Type != EBuildTargetType::Program); 
 
 					if (!bSupportsDefaultBuildTarget && BuildTargetNames.Contains(Target.Name))
 					{
diff --git a/Engine/Source/Programs/AutomationTool/AutomationUtils/ProjectParams.cs b/Engine/Source/Programs/AutomationTool/AutomationUtils/ProjectParams.cs
--- a/Engine/Source/Programs/AutomationTool/AutomationUtils/ProjectParams.cs
+++ b/Engine/Source/Programs/AutomationTool/AutomationUtils/ProjectParams.cs
@@ -2501,7 +2501,7 @@ namespace AutomationTool
 
 				SelectDefaultTarget(AvailableEditorTargets, TargetType.Editor, ref EditorTarget);
 
-				if (AvailableServerTargets.Count > 0 && (DedicatedServer || Cook || CookOnTheFly)) // only if server is needed
+				if (AvailableServerTargets.Count > 0 && DedicatedServer) // only if server is needed
 				{
 					SelectDefaultTarget(AvailableServerTargets, TargetType.Server, ref ServerTarget);
 				}
