Speculative crash fix for closing the editor during Launch On.

#rnx

diff --git a/Engine/Source/Developer/LauncherServices/Private/Launcher/LauncherWorker.cpp b/Engine/Source/Developer/LauncherServices/Private/Launcher/LauncherWorker.cpp
--- a/Engine/Source/Developer/LauncherServices/Private/Launcher/LauncherWorker.cpp
+++ b/Engine/Source/Developer/LauncherServices/Private/Launcher/LauncherWorker.cpp
@@ -104,7 +104,7 @@ uint32 FLauncherWorker::Run( )
 			}
 		}
 
-		if (TaskChain->IsChainFinished())
+		if (TaskChain.IsValid() && TaskChain->IsChainFinished())
 		{
 			Status = ELauncherWorkerStatus::Completed;
 
@@ -148,7 +148,7 @@ uint32 FLauncherWorker::Run( )
 
 		TaskChain->Cancel();
 
-		while (!TaskChain->IsChainFinished())
+		while (TaskChain.IsValid() && !TaskChain->IsChainFinished())
 		{
 			FPlatformProcess::Sleep(0.0);
 		}		
@@ -156,7 +156,7 @@ uint32 FLauncherWorker::Run( )
 
 	FPlatformProcess::ClosePipe(ReadPipe, WritePipe);
 
-	if (Status == ELauncherWorkerStatus::Canceling)
+	if (!TaskChain.IsValid() || Status == ELauncherWorkerStatus::Canceling)
 	{
 		LaunchCanceled.Broadcast(FPlatformTime::Seconds() - LaunchStartTime);
 		Status = ELauncherWorkerStatus::Canceled;
