Fixing bug where we use -fileserver UAT argument when cooked output is stored in zenserver, but we're staging to container/pak files.

diff --git a/Engine/Source/Developer/LauncherServices/Private/Launcher/LauncherWorker.cpp b/Engine/Source/Developer/LauncherServices/Private/Launcher/LauncherWorker.cpp
--- a/Engine/Source/Developer/LauncherServices/Private/Launcher/LauncherWorker.cpp
+++ b/Engine/Source/Developer/LauncherServices/Private/Launcher/LauncherWorker.cpp
@@ -701,7 +701,11 @@ FString FLauncherWorker::CreateUATCommand( const ILauncherProfileRef& InProfile,
 				// TODO: launch the zen server from the client once the external CBTB is done
 				// -fileserver tells UAT to take the cotf/fileserver path and stage a thin client that loads data via the network
 				// -skipserver prevents UAT from launching a COTF server for this CBTB scenario
-				UATCommand += TEXT(" -zenstore -fileserver -skipserver");
+				UATCommand += TEXT(" -zenstore -skipserver");
+				if (!InProfile->IsPackingWithUnrealPak())
+				{
+					UATCommand += TEXT(" -fileserver");
+				}
 			}
 
 			if (FDerivedDataCacheInterface* DDC = TryGetDerivedDataCache())
@@ -862,7 +866,11 @@ FString FLauncherWorker::CreateUATCommand( const ILauncherProfileRef& InProfile,
 			// TODO: launch the zen server from the client once the external CBTB is done
 			// -fileserver tells UAT to take the cotf/fileserver path and stage a thin client that loads data via the network
 			// -skipserver prevents UAT from launching a COTF server for this CBTB scenario
-			UATCommand += TEXT(" -zenstore -fileserver -skipserver");
+			UATCommand += TEXT(" -zenstore -skipserver");
+			if (!InProfile->IsPackingWithUnrealPak())
+			{
+				UATCommand += TEXT(" -fileserver");
+			}
 		}
 		if (InProfile->IsPackingWithUnrealPak())
 		{
