PR #12400: Fix several longstanding instances of "incrimental"

#rnx

- Old misspelt version of the enum is now deprecated as it is in a public header.

diff --git a/Engine/Source/Developer/LauncherServices/Private/Profiles/LauncherProfile.cpp b/Engine/Source/Developer/LauncherServices/Private/Profiles/LauncherProfile.cpp
--- a/Engine/Source/Developer/LauncherServices/Private/Profiles/LauncherProfile.cpp
+++ b/Engine/Source/Developer/LauncherServices/Private/Profiles/LauncherProfile.cpp
@@ -35,8 +35,8 @@ FString LexToStringLocalized(ELauncherProfileValidationErrors::Type Value)
 			return LOCTEXT("NoPackageDirectorySpecified", "The deployment requires a package directory to be specified.").ToString();
 		case ELauncherProfileValidationErrors::NoPlatformSDKInstalled:
 			return LOCTEXT("NoPlatformSDKInstalled", "A required platform SDK is missing.").ToString();
-		case ELauncherProfileValidationErrors::UnversionedAndIncrimental:
-			return LOCTEXT("UnversionedAndIncrimental", "Unversioned build cannot be incremental.").ToString();
+		case ELauncherProfileValidationErrors::UnversionedAndIncremental:
+			return LOCTEXT("UnversionedAndIncremental", "Unversioned build cannot be incremental.").ToString();
 		case ELauncherProfileValidationErrors::GeneratingPatchesCanOnlyRunFromByTheBookCookMode:
 			return LOCTEXT("GeneratingPatchesCanOnlyRunFromByTheBookCookMode", "Generating patch requires cook by the book mode.").ToString();
 		case ELauncherProfileValidationErrors::GeneratingMultiLevelPatchesRequiresGeneratePatch:
diff --git a/Engine/Source/Developer/LauncherServices/Private/Profiles/LauncherProfile.h b/Engine/Source/Developer/LauncherServices/Private/Profiles/LauncherProfile.h
--- a/Engine/Source/Developer/LauncherServices/Private/Profiles/LauncherProfile.h
+++ b/Engine/Source/Developer/LauncherServices/Private/Profiles/LauncherProfile.h
@@ -2752,7 +2752,7 @@ protected:
 
 		if (CookUnversioned && CookIncremental && ((CookMode == ELauncherProfileCookModes::ByTheBook) || (CookMode == ELauncherProfileCookModes::ByTheBookInEditor)))
 		{
-			ValidationErrors.Add(ELauncherProfileValidationErrors::UnversionedAndIncrimental);
+			ValidationErrors.Add(ELauncherProfileValidationErrors::UnversionedAndIncremental);
 		}
 
 
diff --git a/Engine/Source/Developer/LauncherServices/Public/ILauncherProfile.h b/Engine/Source/Developer/LauncherServices/Public/ILauncherProfile.h
--- a/Engine/Source/Developer/LauncherServices/Public/ILauncherProfile.h
+++ b/Engine/Source/Developer/LauncherServices/Public/ILauncherProfile.h
@@ -186,8 +186,11 @@ namespace ELauncherProfileValidationErrors
 		/** The platform SDK is not installed but is required. */
 		NoPlatformSDKInstalled,
 
-		/** The profile has unversioned and incrimental specified these are not compatible together */
-		UnversionedAndIncrimental,
+		/** The profile has unversioned and incremental specified these are not compatible together */
+		UnversionedAndIncremental,
+
+		/** Deprecated, do not use */
+		UnversionedAndIncrimental UE_DEPRECATED(5.6, "Update to use the correctly spelt 'ELauncherProfileValidationErrors::UnversionedAndIncremental' instead") = UnversionedAndIncremental,
 
 		/** generating patch requires cook by the book mode*/
 		GeneratingPatchesCanOnlyRunFromByTheBookCookMode,
diff --git a/Engine/Source/Developer/ProjectLauncher/Private/Widgets/Shared/SProjectLauncherValidation.h b/Engine/Source/Developer/ProjectLauncher/Private/Widgets/Shared/SProjectLauncherValidation.h
--- a/Engine/Source/Developer/ProjectLauncher/Private/Widgets/Shared/SProjectLauncherValidation.h
+++ b/Engine/Source/Developer/ProjectLauncher/Private/Widgets/Shared/SProjectLauncherValidation.h
@@ -118,7 +118,7 @@ public:
 
 		VertBox->AddSlot().AutoHeight()
 			[
-				MakeValidationMessage(TEXT("Icons.Error"), ELauncherProfileValidationErrors::UnversionedAndIncrimental)
+				MakeValidationMessage(TEXT("Icons.Error"), ELauncherProfileValidationErrors::UnversionedAndIncremental)
 			];
 
 		VertBox->AddSlot().AutoHeight()
