Eliminate the use of -fileserver -skipserver when using -zenstore as the staging code already treats presence of -zenstore as implying that we should stage to a thin staged directory equivalent to -fileserver, and it avoids problems when the run task adds -filehostip=foo inappropriately when launching with zen streaming from the editor.

diff --git a/Engine/Source/Developer/LauncherServices/Private/Launcher/LauncherWorker.cpp b/Engine/Source/Developer/LauncherServices/Private/Launcher/LauncherWorker.cpp
--- a/Engine/Source/Developer/LauncherServices/Private/Launcher/LauncherWorker.cpp
+++ b/Engine/Source/Developer/LauncherServices/Private/Launcher/LauncherWorker.cpp
@@ -698,14 +698,7 @@ FString FLauncherWorker::CreateUATCommand( const ILauncherProfileRef& InProfile,
 
 			if (InProfile->IsUsingZenStore())
 			{
-				// TODO: launch the zen server from the client once the external CBTB is done
-				// -fileserver tells UAT to take the cotf/fileserver path and stage a thin client that loads data via the network
-				// -skipserver prevents UAT from launching a COTF server for this CBTB scenario
-				UATCommand += TEXT(" -zenstore -skipserver");
-				if (!InProfile->IsPackingWithUnrealPak())
-				{
-					UATCommand += TEXT(" -fileserver");
-				}
+				UATCommand += TEXT(" -zenstore");
 			}
 
 			if (FDerivedDataCacheInterface* DDC = TryGetDerivedDataCache())
@@ -863,14 +856,7 @@ FString FLauncherWorker::CreateUATCommand( const ILauncherProfileRef& InProfile,
 		UATCommand += TEXT(" -skipcook -CookInEditor"); // don't cook anything the editor is doing it ;)
 		if (InProfile->IsUsingZenStore())
 		{
-			// TODO: launch the zen server from the client once the external CBTB is done
-			// -fileserver tells UAT to take the cotf/fileserver path and stage a thin client that loads data via the network
-			// -skipserver prevents UAT from launching a COTF server for this CBTB scenario
-			UATCommand += TEXT(" -zenstore -skipserver");
-			if (!InProfile->IsPackingWithUnrealPak())
-			{
-				UATCommand += TEXT(" -fileserver");
-			}
+			UATCommand += TEXT(" -zenstore");
 		}
 		if (InProfile->IsPackingWithUnrealPak())
 		{
