Minor tweak to the wording for build target & cook platform error message.

#rnx

diff --git a/Engine/Source/Developer/LauncherServices/Private/Profiles/LauncherProfile.cpp b/Engine/Source/Developer/LauncherServices/Private/Profiles/LauncherProfile.cpp
--- a/Engine/Source/Developer/LauncherServices/Private/Profiles/LauncherProfile.cpp
+++ b/Engine/Source/Developer/LauncherServices/Private/Profiles/LauncherProfile.cpp
@@ -62,7 +62,7 @@ FString LexToStringLocalized(ELauncherProfileValidationErrors::Type Value)
 		case ELauncherProfileValidationErrors::IoStoreRequiresPakFiles:
 			return LOCTEXT("IoStoreRequiresPakFilesError", "UnrealPak must be selected when using I/O store container file(s)").ToString();
 		case ELauncherProfileValidationErrors::BuildTargetCookVariantMismatch:
-			return LOCTEXT("BuildTargetCookVariantMismatch", "Build Target and Cook Variant mismatch.").ToString();
+			return LOCTEXT("BuildTargetCookVariantMismatch", "Build Target and Platform Variant mismatch.").ToString();
 		case ELauncherProfileValidationErrors::BuildTargetIsRequired:
 			return LOCTEXT("BuildTargetIsRequired", "This profile requires an explicit Build Target set.").ToString();
 		case ELauncherProfileValidationErrors::FallbackBuildTargetIsRequired:
