New ILauncherProfile::AssignId function to allow a previously unsaved profile to become saveable

#rnx

diff --git a/Engine/Source/Developer/LauncherServices/Private/Profiles/LauncherProfile.h b/Engine/Source/Developer/LauncherServices/Private/Profiles/LauncherProfile.h
--- a/Engine/Source/Developer/LauncherServices/Private/Profiles/LauncherProfile.h
+++ b/Engine/Source/Developer/LauncherServices/Private/Profiles/LauncherProfile.h
@@ -595,6 +595,15 @@ public:
 		return Id;
 	}
 
+	virtual void AssignId() override
+	{
+		if (!Id.IsValid())
+		{
+			Id = FGuid::NewGuid();
+		}
+	}
+
+
 	virtual FString GetFileName() const override
 	{
 		//toupper for filename so that filepaths can be compared the same on case sensitive and case-insensitive platforms
diff --git a/Engine/Source/Developer/LauncherServices/Public/ILauncherProfile.h b/Engine/Source/Developer/LauncherServices/Public/ILauncherProfile.h
--- a/Engine/Source/Developer/LauncherServices/Public/ILauncherProfile.h
+++ b/Engine/Source/Developer/LauncherServices/Public/ILauncherProfile.h
@@ -488,6 +488,11 @@ public:
 	/** Sets all profile settings to their defaults. */
 	virtual void SetDefaults( ) = 0;
 
+	/** 
+	 * Ensures this profile has an Id, if it doesn't already have one
+	 */
+	virtual void AssignId() = 0;
+
 	/**
 	 * Updates the name of the profile.
 	 *
