update AssignId to include a parameter to override existing Id

#rnx

diff --git a/Engine/Source/Developer/LauncherServices/Private/Profiles/LauncherProfile.h b/Engine/Source/Developer/LauncherServices/Private/Profiles/LauncherProfile.h
--- a/Engine/Source/Developer/LauncherServices/Private/Profiles/LauncherProfile.h
+++ b/Engine/Source/Developer/LauncherServices/Private/Profiles/LauncherProfile.h
@@ -595,9 +595,9 @@ public:
 		return Id;
 	}
 
-	virtual void AssignId() override
+	virtual void AssignId( bool bOverrideExisting ) override
 	{
-		if (!Id.IsValid())
+		if (bOverrideExisting || !Id.IsValid())
 		{
 			Id = FGuid::NewGuid();
 		}
diff --git a/Engine/Source/Developer/LauncherServices/Public/ILauncherProfile.h b/Engine/Source/Developer/LauncherServices/Public/ILauncherProfile.h
--- a/Engine/Source/Developer/LauncherServices/Public/ILauncherProfile.h
+++ b/Engine/Source/Developer/LauncherServices/Public/ILauncherProfile.h
@@ -491,7 +491,7 @@ public:
 	/** 
 	 * Ensures this profile has an Id, if it doesn't already have one
 	 */
-	virtual void AssignId() = 0;
+	virtual void AssignId( bool bOverrideExisting = false) = 0;
 
 	/**
 	 * Updates the name of the profile.
