Project Launcher simplification : reset build target to default when changing the fallback project, rather than trying to find a suitable specific build target.

#rnx

diff --git a/Engine/Source/Developer/LauncherServices/Private/Profiles/LauncherProfileManager.cpp b/Engine/Source/Developer/LauncherServices/Private/Profiles/LauncherProfileManager.cpp
--- a/Engine/Source/Developer/LauncherServices/Private/Profiles/LauncherProfileManager.cpp
+++ b/Engine/Source/Developer/LauncherServices/Private/Profiles/LauncherProfileManager.cpp
@@ -459,28 +459,10 @@ void FLauncherProfileManager::SetProjectPath(const FString& InProjectPath)
 		ProjectPath = InProjectPath;
 		ExplicitBuildTargets = FGameProjectHelper::GetExplicitBuildTargetsForProject(ProjectPath);
 
-		if (ExplicitBuildTargets.Num() == 0)
+		if (!BuildTarget.IsEmpty() && !ExplicitBuildTargets.Contains(BuildTarget))
 		{
 			SetBuildTarget(FString());
 		}
-		else
-		{
-			// pick a default - otherwise the user will have to select one every time
-			FString DefaultBuildTarget = ExplicitBuildTargets[0];
-
-			// favor a Game target
-			const TArray<FTargetInfo>& Targets = FDesktopPlatformModule::Get()->GetTargetsForProject(ProjectPath);
-			for (const FTargetInfo& Target : Targets)
-			{
-				if (Target.Type == EBuildTargetType::Game && ExplicitBuildTargets.Contains(Target.Name))
-				{
-					DefaultBuildTarget = Target.Name;
-					break;
-				}
-			}
-
-			SetBuildTarget(DefaultBuildTarget);
-		}
 
 		for (ILauncherProfilePtr Profile : AllProfiles)
 		{
