[AsyncLoading]

- Enable loose file loading when we detect there is no loader being used so  the user has the best chance to loading their content if no explicit option has been specified.
- Removes the explicit passing of -loosefileloading from the quick launcher as it's no longer necessary
- Restores preventing recursive call to ForceRegenerateClass when using zenloader (legacy loader relies on re-entrance here). From original cl desc: Prevent recursively calling ForceRegenerateClass for the same type. This is to prevent stack overflow when a type has circular references during deferred ExportCreation. Deferred creation is already trying to handle circular references however if there is more than one cycle ForceRegenerateClass would cause a stack overflow. BlueprintSupport in particular takes a slightly different path when FinalizingBlueprints such that is won't regenerate the class if the class is cooked. I believe this breaks the recursion in the non-cooked cases which is why the issue is cook only.
- Will re-enable zenloader as the default loader for loose file loading in a separate change

#rnx

diff --git a/Engine/Source/Developer/LauncherServices/Private/Launcher/LauncherWorker.cpp b/Engine/Source/Developer/LauncherServices/Private/Launcher/LauncherWorker.cpp
--- a/Engine/Source/Developer/LauncherServices/Private/Launcher/LauncherWorker.cpp
+++ b/Engine/Source/Developer/LauncherServices/Private/Launcher/LauncherWorker.cpp
@@ -594,19 +594,10 @@ FString FLauncherWorker::CreateUATCommand( const ILauncherProfileRef& InProfile,
 		*SessionName);
 	}
 
-	// Provide reasonable defaults for loading depending on the user config
-	FString LoaderCommands;
-	if (!InProfile->IsUsingZenStore())
-	{
-		// If not streaming from zenstore, enable loose file reads (only enabled in editor by default)
-		LoaderCommands = TEXT("-loosefileloading");
-	}
-
-	FString AdditionalCommandLine = FString::Printf(TEXT(" -addcmdline=\"%s%s %s %s %s\""),
+	FString AdditionalCommandLine = FString::Printf(TEXT(" -addcmdline=\"%s%s %s %s\""),
 		*SessionCommands,
 		*RoleCommands,
 		*LocalizationCommands,
-		*LoaderCommands,
 		*InProfile->GetAdditionalCommandLineParameters());
 	if (AdditionalCommandLine.EndsWith(TEXT("\"  \"")))
 	{
diff --git a/Engine/Source/Runtime/CoreUObject/Private/Blueprint/BlueprintSupport.cpp b/Engine/Source/Runtime/CoreUObject/Private/Blueprint/BlueprintSupport.cpp
--- a/Engine/Source/Runtime/CoreUObject/Private/Blueprint/BlueprintSupport.cpp
+++ b/Engine/Source/Runtime/CoreUObject/Private/Blueprint/BlueprintSupport.cpp
@@ -1424,13 +1424,17 @@ bool FLinkerLoad::DeferExportCreation(const int32 Index, UObject* Outer)
 	// the midst of doing so somewhere up the callstack)
 	if (!bIsLoadingExportClass || (LoadFlags & LOAD_ResolvingDeferredExports) != 0)
 	{
-		DEFERRED_DEPENDENCY_CHECK(!IsExportBeingResolved(Index));
+		bool bAlreadyResolvingExport = IsExportBeingResolved(Index);
 		FScopedResolvingExportTracker ReentranceGuard(this, Index);
 
 		// we want to be very careful, since we haven't filled in the export yet,
 		// we could get stuck in a recursive loop here (force-finalizing the 
 		// class here ends us back 
-		ForceRegenerateClass(LoadClass);
+		if (!bAlreadyResolvingExport || GetLoaderType() != ELoaderType::ZenLoader)
+		{
+			ForceRegenerateClass(LoadClass);
+		}
+
 		return false;
 	}
 	
diff --git a/Engine/Source/Runtime/CoreUObject/Private/Serialization/AsyncPackageLoader.cpp b/Engine/Source/Runtime/CoreUObject/Private/Serialization/AsyncPackageLoader.cpp
--- a/Engine/Source/Runtime/CoreUObject/Private/Serialization/AsyncPackageLoader.cpp
+++ b/Engine/Source/Runtime/CoreUObject/Private/Serialization/AsyncPackageLoader.cpp
@@ -236,6 +236,13 @@ void InitAsyncThread()
 #endif
 		FIoDispatcher& IoDispatcher = FIoDispatcher::Get();
 		const bool bHasScriptObjectsChunk = FIoDispatcherInternal::HasPackageData();
+
+		// If there is no IODispatcher data, enable the loose file loading path for zenloader so 
+		// the user can attempt to find content on disk, otherwise they will have no loading capabilities
+		if (!bHasScriptObjectsChunk)
+		{
+			FAsyncLoadingThreadSettings::Get().bLooseFileLoadingEnabled = true;
+		}
 		if (!bCommandLineDisabled && (bSettingsEnabled || bCommandLineEnabled))
 		{
 			GPackageLoader.Reset(MakeAsyncPackageLoader2(IoDispatcher));
