Add options for Zen Snapshot Import and Zen Pak Streaming to the launcher services.

- also serializing the 'use zen' option

#rnx

diff --git a/Engine/Source/Developer/LauncherServices/Private/Launcher/LauncherWorker.cpp b/Engine/Source/Developer/LauncherServices/Private/Launcher/LauncherWorker.cpp
--- a/Engine/Source/Developer/LauncherServices/Private/Launcher/LauncherWorker.cpp
+++ b/Engine/Source/Developer/LauncherServices/Private/Launcher/LauncherWorker.cpp
@@ -594,14 +594,23 @@ FString FLauncherWorker::CreateUATCommand( const ILauncherProfileRef& InProfile,
 		*SessionName);
 	}
 
-	FString AdditionalCommandLine = FString::Printf(TEXT(" -addcmdline=\"%s%s %s %s\""),
+	// zen pak streaming commandline
+	FString ZenPakStreamingCommands;
+	if (InProfile->IsUsingZenPakStreaming())
+	{
+		ZenPakStreamingCommands = FString::Printf(TEXT("-ZenPak -ZenWorkspaceSharePath=\\\"%s\\\" "), *InProfile->GetZenPakStreamingPath() );
+	}
+
+	FString AdditionalCommandLine;
+	FString AddCmdLine = FString::Printf(TEXT("%s%s %s %s %s"),
 		*SessionCommands,
 		*RoleCommands,
 		*LocalizationCommands,
+		*ZenPakStreamingCommands,
 		*InProfile->GetAdditionalCommandLineParameters());
-	if (AdditionalCommandLine.EndsWith(TEXT("\"  \"")))
+	if (!AddCmdLine.TrimStartAndEnd().IsEmpty()) // don't bother with -addcmdline if it's empty
 	{
-		AdditionalCommandLine.Reset(); // don't bother with -addcmdline if it's empty
+		AdditionalCommandLine = FString::Printf(TEXT(" -addcmdline=\"%s\""), *AddCmdLine);
 	}
 
 
@@ -701,6 +710,10 @@ FString FLauncherWorker::CreateUATCommand( const ILauncherProfileRef& InProfile,
 			{
 				UATCommand += TEXT(" -zenstore");
 			}
+			if (InProfile->IsImportingZenSnapshot())
+			{
+				UATCommand += TEXT(" -snapshot");
+			}
 
 			if (FDerivedDataCacheInterface* DDC = TryGetDerivedDataCache())
 			{
@@ -811,6 +824,10 @@ FString FLauncherWorker::CreateUATCommand( const ILauncherProfileRef& InProfile,
 			{
 				UATCommand += TEXT(" -zenstore");
 			}
+			if (InProfile->IsImportingZenSnapshot())
+			{
+				UATCommand += TEXT(" -snapshot");
+			}
 
 			if (FDerivedDataCacheInterface* DDC = GetDerivedDataCache())
 			{
@@ -845,6 +862,10 @@ FString FLauncherWorker::CreateUATCommand( const ILauncherProfileRef& InProfile,
 		{
 			UATCommand += TEXT(" -zenstore");
 		}
+		if (InProfile->IsImportingZenSnapshot())
+		{
+			UATCommand += TEXT(" -snapshot");
+		}
 		break;
 	case ELauncherProfileCookModes::ByTheBookInEditor:
 		UATCommand += MapList;
@@ -854,6 +875,10 @@ FString FLauncherWorker::CreateUATCommand( const ILauncherProfileRef& InProfile,
 		{
 			UATCommand += TEXT(" -zenstore");
 		}
+		if (InProfile->IsImportingZenSnapshot())
+		{
+			UATCommand += TEXT(" -snapshot");
+		}
 		if (InProfile->IsPackingWithUnrealPak())
 		{
 			UATCommand += TEXT(" -pak");
@@ -869,6 +894,12 @@ FString FLauncherWorker::CreateUATCommand( const ILauncherProfileRef& InProfile,
 		break;
 	case ELauncherProfileCookModes::DoNotCook:
 		UATCommand += TEXT(" -skipcook");
+		if (InProfile->IsUsingZenPakStreaming())
+		{
+			// @fixme: should auto-configure the zen dynamic workspace in UAT?
+			//UATCommand += TEXT(" -createdynamiczenworkspace=");
+			//UATCommand += InProfile->GetZenPakStreamingPath();
+		}
 		break;
 	}
 
diff --git a/Engine/Source/Developer/LauncherServices/Private/Profiles/LauncherProfile.cpp b/Engine/Source/Developer/LauncherServices/Private/Profiles/LauncherProfile.cpp
--- a/Engine/Source/Developer/LauncherServices/Private/Profiles/LauncherProfile.cpp
+++ b/Engine/Source/Developer/LauncherServices/Private/Profiles/LauncherProfile.cpp
@@ -6,7 +6,7 @@
 
 FString LexToStringLocalized(ELauncherProfileValidationErrors::Type Value)
 {
-	static_assert(ELauncherProfileValidationErrors::Count == 30, "GetLocalizedValidationErrorMessage() needs to be updated to account for modified enum values.");
+	static_assert(ELauncherProfileValidationErrors::Count == 31, "GetLocalizedValidationErrorMessage() needs to be updated to account for modified enum values.");
 	switch (Value)
 	{
 		case ELauncherProfileValidationErrors::CopyToDeviceRequiresCookByTheBook:
@@ -69,6 +69,8 @@ FString LexToStringLocalized(ELauncherProfileValidationErrors::Type Value)
 			return LOCTEXT("FallbackBuildTargetIsRequired", "An explicit Default Build Target is required for the selected Variant.").ToString();
 		case ELauncherProfileValidationErrors::CopyToDeviceRequiresNoPackaging:
 			return LOCTEXT("CopyToDeviceRequiresNoPackaging", "Deployment by copying to device requires 'Do not package' packaging.").ToString();
+		case ELauncherProfileValidationErrors::ZenPakStreamingRequiresDeployAndLaunch:
+			return LOCTEXT("ZenPakStreamingRequiresDeployAndLaunch", "Zen Pak Streaming requires deployment to a device and launching.").ToString();
 		default:
 			return TEXT("Unknown");
 	};
diff --git a/Engine/Source/Developer/LauncherServices/Private/Profiles/LauncherProfile.h b/Engine/Source/Developer/LauncherServices/Private/Profiles/LauncherProfile.h
--- a/Engine/Source/Developer/LauncherServices/Private/Profiles/LauncherProfile.h
+++ b/Engine/Source/Developer/LauncherServices/Private/Profiles/LauncherProfile.h
@@ -54,8 +54,9 @@ enum ELauncherVersion
 	LAUNCHERSERVICES_REMOVEDNUMCOOKERSTOSPAWN = 32,
 	LAUNCHERSERVICES_ADDEDORIGINALRELEASEVERSION = 33,
 	LAUNCHERSERVICES_ADDBUILDTARGETNAME = 34,
-	LAUNCHERSERVICED_ADDEDRETAINSTAGEDDIRECTORY = 35,
-	LAUNCHERSERVICED_REMOVEDRETAINSTAGEDDIRECTORY = 36,
+	LAUNCHERSERVICES_ADDEDRETAINSTAGEDDIRECTORY = 35,
+	LAUNCHERSERVICES_REMOVEDRETAINSTAGEDDIRECTORY = 36,
+	LAUNCHERSERVICES_ADDZENOPTIONS = 37,
 	//ADD NEW STUFF HERE
 
 
@@ -1108,12 +1109,20 @@ public:
 			Archive << BuildTargetName;
 		}
 
-		if (Version >= LAUNCHERSERVICED_ADDEDRETAINSTAGEDDIRECTORY && Version < LAUNCHERSERVICED_REMOVEDRETAINSTAGEDDIRECTORY)
+		if (Version >= LAUNCHERSERVICES_ADDEDRETAINSTAGEDDIRECTORY && Version < LAUNCHERSERVICES_REMOVEDRETAINSTAGEDDIRECTORY)
 		{
 			bool bRetainStagedDirectory = false;
 			Archive << bRetainStagedDirectory;
 		}
 
+		if (Version >= LAUNCHERSERVICES_ADDZENOPTIONS)
+		{
+			Archive << bUseZenStore;
+			Archive << bImportingZenSnapshot;
+			Archive << bUseZenPakStreaming;
+			Archive << ZenPakStreamingPath;
+		}
+
 		DefaultLaunchRole->Serialize(Archive);
 
 		// serialize launch roles
@@ -1259,6 +1268,11 @@ public:
 		Writer.WriteValue("MakeBinaryConfig", bMakeBinaryConfig);
 		Writer.WriteValue("BuildTargetSpecified", BuildTargetSpecified);
 		Writer.WriteValue("BuildTargetName", BuildTargetName);
+		Writer.WriteValue("UseZenStore", bUseZenStore);
+		Writer.WriteValue("ImportingZenSnapshot", bImportingZenSnapshot);
+		Writer.WriteValue("UseZenPakStreaming", bUseZenPakStreaming);
+		Writer.WriteValue("ZenPakStreamingPath", ZenPakStreamingPath);
+
 
 		// serialize the default launch role
 		DefaultLaunchRole->Save(Writer, TEXT("DefaultRole"));
@@ -1962,6 +1976,14 @@ public:
 			BuildTargetName = Object.GetStringField(TEXT("BuildTargetName"));
 		}
 
+		if (Version >= LAUNCHERSERVICES_ADDZENOPTIONS)
+		{
+			bUseZenStore = Object.GetBoolField(TEXT("UseZenStore"));
+			bImportingZenSnapshot = Object.GetBoolField(TEXT("ImportingZenSnapshot"));
+			bUseZenPakStreaming = Object.GetBoolField(TEXT("UseZenPakStreaming"));
+			ZenPakStreamingPath = Object.GetStringField(TEXT("ZenPakStreamingPath"));
+		}
+
 		// load the default launch role
 		TSharedPtr<FJsonObject> Role = Object.GetObjectField(TEXT("DefaultRole"));
 		DefaultLaunchRole->Load(*(Role.Get()));
@@ -2109,6 +2131,8 @@ public:
 		bNotForLicensees = false;
 		bUseIoStore = false;
 		bUseZenStore = false;
+		bImportingZenSnapshot = false;
+		bUseZenPakStreaming = false;
 		bShouldUpdateFlash = false;
 		bIsDeviceASimulator = false;
 		bMakeBinaryConfig = false;
@@ -2629,6 +2653,59 @@ public:
 		return bUseZenStore;
 	}
 
+	virtual void SetImportingZenSnapshot( bool Import ) override
+	{
+		if (bImportingZenSnapshot != Import)
+		{
+			bImportingZenSnapshot = Import;
+
+			Validate();
+		}
+	}
+
+	virtual bool IsImportingZenSnapshot() const override
+	{
+		if (CookMode != ELauncherProfileCookModes::DoNotCook)
+		{
+			return bImportingZenSnapshot;
+		}
+		return false;
+	}
+
+	virtual void SetUseZenPakStreaming( bool UseZenPakStreaming ) override
+	{
+		if (bUseZenPakStreaming != UseZenPakStreaming)
+		{
+			bUseZenPakStreaming = UseZenPakStreaming;
+
+			Validate();
+		}
+	}
+
+	virtual bool IsUsingZenPakStreaming() const override
+	{
+		if (CookMode == ELauncherProfileCookModes::DoNotCook)
+		{
+			return bUseZenPakStreaming;
+		}
+		return false;
+	}
+
+	virtual void SetZenPakStreamingPath( const FString& Path ) override
+	{
+		if (ZenPakStreamingPath != Path)
+		{
+			ZenPakStreamingPath = Path;
+
+			Validate();
+		}
+	}
+
+	virtual FString GetZenPakStreamingPath() const override
+	{
+		return ZenPakStreamingPath;
+	}
+
 	virtual void SetShouldUpdateDeviceFlash(bool bInShouldUpdateFlash) override
 	{
 		bShouldUpdateFlash = bInShouldUpdateFlash;
@@ -2722,7 +2799,7 @@ protected:
 		}
 
 		// Deploy: deployment by copying to devices requires cooking by the book
-		if ((DeploymentMode == ELauncherProfileDeploymentModes::CopyToDevice) && ((CookMode != ELauncherProfileCookModes::ByTheBook)&&(CookMode!=ELauncherProfileCookModes::ByTheBookInEditor)))
+		if ((DeploymentMode == ELauncherProfileDeploymentModes::CopyToDevice) && ((CookMode != ELauncherProfileCookModes::ByTheBook)&&(CookMode!=ELauncherProfileCookModes::ByTheBookInEditor)) && !IsUsingZenPakStreaming())
 		{
 			ValidationErrors.Add(ELauncherProfileValidationErrors::CopyToDeviceRequiresCookByTheBook);
 		}
@@ -2759,6 +2836,12 @@ protected:
 			}
 		}
 
+		// Launch: Zen pak streaming requires deployment and launch
+		if (IsUsingZenPakStreaming() && (LaunchMode == ELauncherProfileLaunchModes::DoNotLaunch || DeploymentMode != ELauncherProfileDeploymentModes::CopyToDevice))
+		{
+			ValidationErrors.Add(ELauncherProfileValidationErrors::ZenPakStreamingRequiresDeployAndLaunch);
+		}
+
 		if (CookUnversioned && CookIncremental && ((CookMode == ELauncherProfileCookModes::ByTheBook) || (CookMode == ELauncherProfileCookModes::ByTheBookInEditor)))
 		{
 			ValidationErrors.Add(ELauncherProfileValidationErrors::UnversionedAndIncremental);
@@ -3245,6 +3328,9 @@ private:
 	// otherwise this is an absolute path.
 	FString ShareableProjectPath;
 
+	// Holds the full path to a folder to use for Zen pak streaming
+	FString ZenPakStreamingPath;
+
 	// Holds the collection of validation errors.
 	TArray<ELauncherProfileValidationErrors::Type> ValidationErrors;
 	FString InvalidPlatform;
@@ -3270,6 +3356,12 @@ private:
 	// Use Zen storage server
 	bool bUseZenStore;
 
+	// Import a Zen snapshot before cooking
+	bool bImportingZenSnapshot;
+
+	// Use Zen pak streaming
+	bool bUseZenPakStreaming;
+
 	// Make binary config.
 	bool bMakeBinaryConfig;
 
diff --git a/Engine/Source/Developer/LauncherServices/Public/ILauncherProfile.h b/Engine/Source/Developer/LauncherServices/Public/ILauncherProfile.h
--- a/Engine/Source/Developer/LauncherServices/Public/ILauncherProfile.h
+++ b/Engine/Source/Developer/LauncherServices/Public/ILauncherProfile.h
@@ -234,12 +234,15 @@ namespace ELauncherProfileValidationErrors
 		/** Build Target is required */
 		BuildTargetIsRequired,
 
-		/** FallbackBuild Target is required */
+		/** Fallback Build Target is required */
 		FallbackBuildTargetIsRequired,
 
 		/** Packaging and deploying are mutually exclusive */
 		CopyToDeviceRequiresNoPackaging,
 
+		/** Zen Pak Streaming requires deployment and launching */
+		ZenPakStreamingRequiresDeployAndLaunch,
+
 		Count
 	};
 }
@@ -1366,7 +1369,7 @@ public:
 	virtual bool IsUsingIoStore() const = 0;
 
 	/**
-	 * Sets whether to use the Zen storage server
+	 * Sets whether to use the Zen storage server. Note, this property is ignored if the project itself is configured to use Zen store (via UProjectPackagingSettings::bUseZenStore)
 	 * @param bUseZenStore Whether to use the Zen storage server
 	 */
 	virtual void SetUseZenStore(bool bUseZenStore) = 0;
@@ -1378,6 +1381,49 @@ public:
 	 */
 	virtual bool IsUsingZenStore() const = 0;
 
+	/**
+	 * Set whether to import a Zen snapshot before cooking.
+	 * 
+	 * @param Import Whether to import a Zen snapshot
+	 */
+	virtual void SetImportingZenSnapshot( bool Import ) = 0;
+
+	/**
+	 * Checks whether a Zen snapshot will be imported before cooking.
+	 * 
+	 * @return true if importing a Zen snapshot, false otherwise
+	 */
+	virtual bool IsImportingZenSnapshot() const = 0;
+
+	/**
+	 * Set whether to use Zen pak streaming.
+	 * 
+	 * @param UseZenPakStreaming Whether use Zen pak streaming.
+	 */
+	virtual void SetUseZenPakStreaming( bool UseZenPakStreaming ) = 0;
+
+	/**
+	 * Checks whether Zen pak streaming will be used.
+	 * 
+	 * @return true if using Zen pak streaming, false otherwise
+	 */
+	virtual bool IsUsingZenPakStreaming() const = 0;
+
+	/**
+	 * Sets whether to use Zen pak streaming from the given path.
+	 * Assumes the path contains a suitable build
+	 * 
+	 * @param Path Path to the build to use for Zen pak streaming
+	 */
+	virtual void SetZenPakStreamingPath( const FString& Path ) = 0;
+
+	/**
+	 * Get the current Zen pak streaming path, if it is being used.
+	 * 
+	 * @return current Zen pak streaming path, if any
+	 */
+	virtual FString GetZenPakStreamingPath() const = 0;
+
 	/**
 	 * Sets whether to make a binary config file during packaging
 	 * @param bMakeBinaryConfig Whether to make a binary config file during staging
diff --git a/Engine/Source/Developer/ProjectLauncher/Private/Widgets/Shared/SProjectLauncherValidation.h b/Engine/Source/Developer/ProjectLauncher/Private/Widgets/Shared/SProjectLauncherValidation.h
--- a/Engine/Source/Developer/ProjectLauncher/Private/Widgets/Shared/SProjectLauncherValidation.h
+++ b/Engine/Source/Developer/ProjectLauncher/Private/Widgets/Shared/SProjectLauncherValidation.h
@@ -196,6 +196,11 @@ public:
 				MakeValidationMessage(TEXT("Icons.Error"), ELauncherProfileValidationErrors::CopyToDeviceRequiresNoPackaging)
 			];
 
+		VertBox->AddSlot().AutoHeight()
+			[
+				MakeValidationMessage(TEXT("Icons.Error"), ELauncherProfileValidationErrors::ZenPakStreamingRequiresDeployAndLaunch)
+			];
+
 		check(VertBox->NumSlots() == ELauncherProfileValidationErrors::Count);
 
 		ChildSlot[VertBox.ToSharedRef()];
