Add architecture selection to launcher profile

#rnx

diff --git a/Engine/Source/Developer/LauncherServices/Private/Launcher/LauncherWorker.cpp b/Engine/Source/Developer/LauncherServices/Private/Launcher/LauncherWorker.cpp
--- a/Engine/Source/Developer/LauncherServices/Private/Launcher/LauncherWorker.cpp
+++ b/Engine/Source/Developer/LauncherServices/Private/Launcher/LauncherWorker.cpp
@@ -472,10 +472,17 @@ FString FLauncherWorker::CreateUATCommand( const ILauncherProfileRef& InProfile,
 		{
 			OptionalParams += TEXT(" -noclient");
 		}
+
+		if (InProfile->GetServerArchitectures().Num() > 0)
+		{
+			ServerCommand += TEXT(" -serverarchitecture=") + FString::Join(InProfile->GetServerArchitectures(), TEXT("+") );
+		}
 	}
+	bool bSetClientArchitecture = false;
 	if (Platforms.Len() > 0)
 	{
 		PlatformCommand = TEXT(" -platform=") + Platforms.RightChop(1);
+		bSetClientArchitecture = true;
 	}
 	
 	UATCommand += PlatformCommand;
@@ -502,10 +509,17 @@ FString FLauncherWorker::CreateUATCommand( const ILauncherProfileRef& InProfile,
 		if (Platforms.Contains(TEXT("IOS")))
 		{
 			UATCommand += TEXT(" -clientarchitecture=iossimulator");
+			bSetClientArchitecture = false;
 		}
 		// TODO: add tvOS and VisionOS simulators below
 	}
 
+	if (bSetClientArchitecture && InProfile->GetClientArchitectures().Num() > 0)
+	{
+		UATCommand += TEXT(" -clientarchitecture=") + FString::Join(InProfile->GetClientArchitectures(), TEXT("+") );
+	}
+
+
 	// device list
 	FString DeviceNames = TEXT("");
 	FString DeviceCommand = TEXT("");
diff --git a/Engine/Source/Developer/LauncherServices/Private/Profiles/LauncherProfile.h b/Engine/Source/Developer/LauncherServices/Private/Profiles/LauncherProfile.h
--- a/Engine/Source/Developer/LauncherServices/Private/Profiles/LauncherProfile.h
+++ b/Engine/Source/Developer/LauncherServices/Private/Profiles/LauncherProfile.h
@@ -57,6 +57,7 @@ enum ELauncherVersion
 	LAUNCHERSERVICES_ADDEDRETAINSTAGEDDIRECTORY = 35,
 	LAUNCHERSERVICES_REMOVEDRETAINSTAGEDDIRECTORY = 36,
 	LAUNCHERSERVICES_ADDZENOPTIONS = 37,
+	LAUNCHERSERVICES_ADDBUILDARCHITECTURE = 38,
 	//ADD NEW STUFF HERE
 
 
@@ -1123,6 +1124,13 @@ public:
 			Archive << ZenPakStreamingPath;
 		}
 
+		if (Version >= LAUNCHERSERVICES_ADDBUILDARCHITECTURE)
+		{
+			Archive << ClientArchitectures;
+			Archive << ServerArchitectures;
+			Archive << EditorArchitectures;
+		}
+
 		DefaultLaunchRole->Serialize(Archive);
 
 		// serialize launch roles
@@ -1272,6 +1280,9 @@ public:
 		Writer.WriteValue("ImportingZenSnapshot", bImportingZenSnapshot);
 		Writer.WriteValue("UseZenPakStreaming", bUseZenPakStreaming);
 		Writer.WriteValue("ZenPakStreamingPath", ZenPakStreamingPath);
+		Writer.WriteValue("ClientArchitectures", ClientArchitectures);
+		Writer.WriteValue("ServerArchitectures", ServerArchitectures);
+		Writer.WriteValue("EditorArchitectures", EditorArchitectures);
 
 
 		// serialize the default launch role
@@ -1598,6 +1609,19 @@ public:
 			}
 		}
 
+		if (ClientArchitectures.Num() > 0)
+		{
+			Writer.WriteValue("clientarchitecture", FString::Join(ClientArchitectures, TEXT("+")));
+		}
+		if (ServerArchitectures.Num() > 0)
+		{
+			Writer.WriteValue("serverarchitecture", FString::Join(ServerArchitectures, TEXT("+")));
+		}
+		if (EditorArchitectures.Num() > 0)
+		{
+			Writer.WriteValue("editorarchitecture", FString::Join(EditorArchitectures, TEXT("+")));
+		}
+
 		/*
 		"script", ""
 		"project", ""
@@ -1984,6 +2008,16 @@ public:
 			ZenPakStreamingPath = Object.GetStringField(TEXT("ZenPakStreamingPath"));
 		}
 
+		if (Version >= LAUNCHERSERVICES_ADDBUILDARCHITECTURE)
+		{
+			ClientArchitectures.Reset();
+			ServerArchitectures.Reset();
+			EditorArchitectures.Reset();
+			Object.TryGetStringArrayField(TEXT("ClientArchitectures"), ClientArchitectures);
+			Object.TryGetStringArrayField(TEXT("ServerArchitectures"), ServerArchitectures);
+			Object.TryGetStringArrayField(TEXT("EditorArchitectures"), EditorArchitectures);
+		}
+
 		// load the default launch role
 		TSharedPtr<FJsonObject> Role = Object.GetObjectField(TEXT("DefaultRole"));
 		DefaultLaunchRole->Load(*(Role.Get()));
@@ -2752,6 +2786,36 @@ public:
 		return ExplictBuildTargetNames.Num() > 0;
 	}
 
+	virtual void SetClientArchitectures( const TArray<FString>& InArchitectures )
+	{
+		ClientArchitectures = InArchitectures;
+	}
+
+	virtual const TArray<FString>& GetClientArchitectures() const override
+	{
+		return ClientArchitectures;
+	}
+
+	virtual void SetServerArchitectures( const TArray<FString>& InArchitectures )
+	{
+		ServerArchitectures = InArchitectures;
+	}
+
+	virtual const TArray<FString>& GetServerArchitectures() const override
+	{
+		return ServerArchitectures;
+	}
+
+	virtual void SetEditorArchitectures( const TArray<FString>& InArchitectures )
+	{
+		EditorArchitectures = InArchitectures;
+	}
+
+	virtual const TArray<FString>& GetEditorArchitectures() const override
+	{
+		return EditorArchitectures;
+	}
+
 
 	//~ End ILauncherProfile Interface
 
@@ -3371,6 +3435,11 @@ private:
 	// Is the launch device actually a Simulator
 	bool bIsDeviceASimulator;
 
+	// Architectures to build (x64, arm64 etc)
+	TArray<FString> ClientArchitectures;
+	TArray<FString> ServerArchitectures;
+	TArray<FString> EditorArchitectures;
+
 private:
 
 	// cook in the editor callbacks (not valid for any other cook mode)
diff --git a/Engine/Source/Developer/LauncherServices/Public/ILauncherProfile.h b/Engine/Source/Developer/LauncherServices/Public/ILauncherProfile.h
--- a/Engine/Source/Developer/LauncherServices/Public/ILauncherProfile.h
+++ b/Engine/Source/Developer/LauncherServices/Public/ILauncherProfile.h
@@ -1456,6 +1456,36 @@ public:
 	 */
 	virtual bool IsDeviceASimulator() const = 0;
 
+	/** 
+	 * Sets the client architecture(s) to build (x64, arm64 etc). Leave empty to use the default
+	 */
+	virtual void SetClientArchitectures( const TArray<FString>& InArchitectures ) = 0;
+
+	/**
+	 * Returns the client architecture(s) that will be build. empty means the default will be used
+	 */
+	virtual const TArray<FString>& GetClientArchitectures() const = 0;
+	
+	/** 
+	 * Sets the server architecture(s) to build (x64, arm64 etc). Leave empty to use the default
+	 */
+	virtual void SetServerArchitectures( const TArray<FString>& InArchitectures ) = 0;
+
+	/**
+	 * Returns the server architecture(s) that will be build. empty means the default will be used
+	 */
+	virtual const TArray<FString>& GetServerArchitectures() const = 0;
+
+	/** 
+	 * Sets the editor architecture(s) to build (x64, arm64 etc). Leave empty to use the default
+	 */
+	virtual void SetEditorArchitectures( const TArray<FString>& InArchitectures ) = 0;
+
+	/**
+	 * Returns the editor architecture(s) that will be build. empty means the default will be used
+	 */
+	virtual const TArray<FString>& GetEditorArchitectures() const = 0;
+
 public:
 	/**
 	 * Helper function to get all of the build targets available for this profile, based on its current project & cook platforms
