Adding hooks for extensibility:
- callback for post-processing launcher worker command line parameters.
- custom key-value pairs that are saved in the profile.

#rnx

diff --git a/Engine/Source/Developer/LauncherServices/Private/Launcher/LauncherWorker.cpp b/Engine/Source/Developer/LauncherServices/Private/Launcher/LauncherWorker.cpp
--- a/Engine/Source/Developer/LauncherServices/Private/Launcher/LauncherWorker.cpp
+++ b/Engine/Source/Developer/LauncherServices/Private/Launcher/LauncherWorker.cpp
@@ -3,6 +3,7 @@
 #include "Launcher/LauncherWorker.h"
 #include "HAL/PlatformTime.h"
 #include "HAL/FileManager.h"
+#include "ILauncherServicesModule.h"
 #include "ISourceCodeAccessor.h"
 #include "ISourceCodeAccessModule.h"
 #include "ITargetDeviceProxy.h"
@@ -389,6 +390,8 @@ static FString Join(const TSet<FString>& Tokens, const FString& Delimeter)
 
 FString FLauncherWorker::CreateUATCommand( const ILauncherProfileRef& InProfile, const TArray<FString>& InPlatforms, TArray<FCommandDesc>& OutCommands, FString& CommandStart, bool bForTurnkeyCustomBuild )
 {
+	ILauncherProfileManagerRef LauncherProfileManager = FModuleManager::LoadModuleChecked<ILauncherServicesModule>("LauncherServices").GetProfileManager();
+
 	CommandStart = TEXT("");
 	FString UATCommand = bForTurnkeyCustomBuild ? TEXT("") : TEXT(" -utf8output");
 	FGuid SessionId(FGuid::NewGuid());
@@ -608,6 +611,10 @@ FString FLauncherWorker::CreateUATCommand( const ILauncherProfileRef& InProfile,
 		*SessionName);
 	}
 
+	// allow external items to adjust the command line
+	FString ProfileAdditionalCommandLineParameters = InProfile->GetAdditionalCommandLineParameters();
+	LauncherProfileManager->OnPostProcessLaunchCommandLine().Broadcast(InProfile, ProfileAdditionalCommandLineParameters);
+
 	FString AdditionalCommandLine;
 	TStringBuilder<64> InheritableGameOptions;
 	ECommandLineArgumentFlags CommandLineContextFlags = ECommandLineArgumentFlags::ClientContext | ECommandLineArgumentFlags::ServerContext;
@@ -618,7 +625,7 @@ FString FLauncherWorker::CreateUATCommand( const ILauncherProfileRef& InProfile,
 		*SessionCommands,
 		*RoleCommands,
 		*LocalizationCommands,
-		*InProfile->GetAdditionalCommandLineParameters(),
+		*ProfileAdditionalCommandLineParameters,
 		*EscapedInheritableGameOptions);
 	if (!AddCmdLine.TrimStartAndEnd().IsEmpty()) // don't bother with -addcmdline if it's empty
 	{
diff --git a/Engine/Source/Developer/LauncherServices/Private/Profiles/LauncherProfile.h b/Engine/Source/Developer/LauncherServices/Private/Profiles/LauncherProfile.h
--- a/Engine/Source/Developer/LauncherServices/Private/Profiles/LauncherProfile.h
+++ b/Engine/Source/Developer/LauncherServices/Private/Profiles/LauncherProfile.h
@@ -58,6 +58,7 @@ enum ELauncherVersion
 	LAUNCHERSERVICES_REMOVEDRETAINSTAGEDDIRECTORY = 36,
 	LAUNCHERSERVICES_ADDZENOPTIONS = 37,
 	LAUNCHERSERVICES_ADDBUILDARCHITECTURE = 38,
+	LAUNCHERSERVICES_ADDCUSTOMPROPERTIES = 39,
 	//ADD NEW STUFF HERE
 
 
@@ -1131,6 +1132,12 @@ public:
 			Archive << EditorArchitectures;
 		}
 
+		if (Version >= LAUNCHERSERVICES_ADDCUSTOMPROPERTIES)
+		{
+			Archive << CustomStringProperties;
+			Archive << CustomBoolProperties;
+		}
+
 		DefaultLaunchRole->Serialize(Archive);
 
 		// serialize launch roles
@@ -1283,6 +1290,8 @@ public:
 		Writer.WriteValue("ClientArchitectures", ClientArchitectures);
 		Writer.WriteValue("ServerArchitectures", ServerArchitectures);
 		Writer.WriteValue("EditorArchitectures", EditorArchitectures);
+		Writer.WriteValue("CustomStringProperties", CustomStringProperties);
+		Writer.WriteValue("CustomBoolProperties", CustomBoolProperties);
 
 
 		// serialize the default launch role
@@ -2018,6 +2027,30 @@ public:
 			Object.TryGetStringArrayField(TEXT("EditorArchitectures"), EditorArchitectures);
 		}
 
+		if (Version >= LAUNCHERSERVICES_ADDCUSTOMPROPERTIES)
+		{
+			const TSharedPtr<FJsonObject>* CustomPropertiesObject;
+
+			CustomStringProperties.Reset();
+			if (Object.TryGetObjectField(TEXT("CustomStringProperties"), CustomPropertiesObject))
+			{
+				for (const auto& KeyPairs : (*CustomPropertiesObject)->Values)
+				{
+					CustomStringProperties.Add(KeyPairs.Key, KeyPairs.Value->AsString());
+				}
+			}
+
+			CustomBoolProperties.Reset();
+			if (Object.TryGetObjectField(TEXT("CustomBoolProperties"), CustomPropertiesObject))
+			{
+				for (const auto& KeyPairs : (*CustomPropertiesObject)->Values)
+				{
+					CustomBoolProperties.Add(KeyPairs.Key, KeyPairs.Value->AsBool());
+				}
+			}
+		}
+
+
 		// load the default launch role
 		TSharedPtr<FJsonObject> Role = Object.GetObjectField(TEXT("DefaultRole"));
 		DefaultLaunchRole->Load(*(Role.Get()));
@@ -2816,6 +2849,16 @@ public:
 		return EditorArchitectures;
 	}
 
+	virtual TMap<FString,FString>& GetCustomStringProperties() override
+	{
+		return CustomStringProperties;
+	}
+
+	virtual TMap<FString, bool>& GetCustomBoolProperties() override
+	{
+		return CustomBoolProperties;
+	}
+
 
 	//~ End ILauncherProfile Interface
 
@@ -3440,6 +3483,10 @@ private:
 	TArray<FString> ServerArchitectures;
 	TArray<FString> EditorArchitectures;
 
+	// Key-Value pairs for extsibility. Not user-facing.
+	TMap<FString,FString> CustomStringProperties;
+	TMap<FString, bool>   CustomBoolProperties;
+
 private:
 
 	// cook in the editor callbacks (not valid for any other cook mode)
diff --git a/Engine/Source/Developer/LauncherServices/Private/Profiles/LauncherProfileManager.h b/Engine/Source/Developer/LauncherServices/Private/Profiles/LauncherProfileManager.h
--- a/Engine/Source/Developer/LauncherServices/Private/Profiles/LauncherProfileManager.h
+++ b/Engine/Source/Developer/LauncherServices/Private/Profiles/LauncherProfileManager.h
@@ -85,6 +85,11 @@ public:
 		return ProjectChangedDelegate;
 	}
 
+	virtual FOnLauncherProfilePostProcessCommandLine& OnPostProcessLaunchCommandLine( ) override
+	{
+		return ProcessProcessLaunchCommandLine;
+	}
+
 	virtual void RemoveDeviceGroup( const ILauncherDeviceGroupRef& DeviceGroup ) override;
 
 	virtual void SaveDeviceGroups() override;
@@ -211,4 +216,8 @@ private:
 	// Holds a delegate to be invoked when the project changes.
 	FOnLauncherProfileManagerProjectChanged ProjectChangedDelegate;	
 
+	// Holds a delegate to be invoked to modify the command line before launching
+	FOnLauncherProfilePostProcessCommandLine ProcessProcessLaunchCommandLine;
+
+
 };
diff --git a/Engine/Source/Developer/LauncherServices/Public/ILauncherProfile.h b/Engine/Source/Developer/LauncherServices/Public/ILauncherProfile.h
--- a/Engine/Source/Developer/LauncherServices/Public/ILauncherProfile.h
+++ b/Engine/Source/Developer/LauncherServices/Public/ILauncherProfile.h
@@ -4,6 +4,7 @@
 
 #include "CoreMinimal.h"
 #include "Misc/Guid.h"
+#include "Containers/Map.h"
 #include "ILauncherDeviceGroup.h"
 #include "ILauncherProfileLaunchRole.h"
 
@@ -1497,6 +1498,16 @@ public:
 	 */
 	virtual bool RequiresExplicitBuildTargetName() const = 0;
 
+	/**
+	 * Key-Value pairs for extsibility. Not user-facing.
+	 */
+	virtual TMap<FString,FString>& GetCustomStringProperties() = 0;
+
+	/**
+	 * Key-Value pairs for extsibility. Not user-facing.
+	 */
+	virtual TMap<FString, bool>& GetCustomBoolProperties() = 0;
+
 public:
 
 	/** Virtual destructor. */
diff --git a/Engine/Source/Developer/LauncherServices/Public/ILauncherProfileManager.h b/Engine/Source/Developer/LauncherServices/Public/ILauncherProfileManager.h
--- a/Engine/Source/Developer/LauncherServices/Public/ILauncherProfileManager.h
+++ b/Engine/Source/Developer/LauncherServices/Public/ILauncherProfileManager.h
@@ -50,6 +50,14 @@ DECLARE_MULTICAST_DELEGATE_OneParam(FOnLauncherProfileManagerProfileRemoved, con
  */
 DECLARE_MULTICAST_DELEGATE(FOnLauncherProfileManagerProjectChanged);
 
+/**
+ * Declares a delegate to be invoked to post-process a command line before launching
+ *
+ * The first parameter is the profile that is being launched, the second parameter is the command line to modify.
+ */
+DECLARE_MULTICAST_DELEGATE_TwoParams(FOnLauncherProfilePostProcessCommandLine, const ILauncherProfileRef&, FString& );
+
+
 /** Type definition for shared references to instances of ILauncherProfile. */
 typedef TSharedPtr<class ILauncherProfileWizard> ILauncherProfileWizardPtr;
 
@@ -400,6 +408,14 @@ public:
 	virtual FOnLauncherProfileManagerProjectChanged& OnProjectChanged( ) = 0;
 	
 
+	/**
+	 * Returns a delegate that is invoked to post-process the command line before a profile launched
+	 * 
+	 * @return The delegate.
+	 */
+	virtual FOnLauncherProfilePostProcessCommandLine& OnPostProcessLaunchCommandLine( ) = 0;
+
+
 public:
 
 	/** Virtual destructor. */
