Changes to reflect the new incremental cook and zen snapshot import options:

- new Launcher Task for Zen snapshot import.
- deprecate ILauncherProfile's IncrementalCook boolean in favor of a new IncrementalCookMode.
- existing use cases still using the old 'iterative cook' option.
- CustomLaunchUI automatically adjusts  the Unversioned cook property when IncrementalCookMode changes.

#rnx

diff --git a/Engine/Plugins/Developer/CustomLaunchUI/Source/CustomLaunchUI/Private/Marshalling/CustomProfileMarshaller.cpp b/Engine/Plugins/Developer/CustomLaunchUI/Source/CustomLaunchUI/Private/Marshalling/CustomProfileMarshaller.cpp
--- a/Engine/Plugins/Developer/CustomLaunchUI/Source/CustomLaunchUI/Private/Marshalling/CustomProfileMarshaller.cpp
+++ b/Engine/Plugins/Developer/CustomLaunchUI/Source/CustomLaunchUI/Private/Marshalling/CustomProfileMarshaller.cpp
@@ -29,10 +29,10 @@ namespace CustomLaunchUI
 		AddGenerateChunksProperty(ContentSchemeHeader);
 		AddImportZenSnapshotProperty(ContentSchemeHeader);
 		AddZenPakStreamingPathProperty(ContentSchemeHeader);
-		AddIncrementalCookProperty(ContentSchemeHeader);
 
 		FLaunchProfileTreeNode& MapsAndCookingHeader = TreeData->AddHeading( LOCTEXT("CookingHeading", "Maps And Cooking") );
 		AddCookProperty(MapsAndCookingHeader);
+		AddIncrementalCookProperty(MapsAndCookingHeader);
 		AddMapsToCookProperty(MapsAndCookingHeader);
 		AddAdditionalCookerOptionsProperty(MapsAndCookingHeader);
 
diff --git a/Engine/Plugins/Developer/CustomLaunchUI/Source/CustomLaunchUI/Private/Marshalling/GenericProfileMarshaller.cpp b/Engine/Plugins/Developer/CustomLaunchUI/Source/CustomLaunchUI/Private/Marshalling/GenericProfileMarshaller.cpp
--- a/Engine/Plugins/Developer/CustomLaunchUI/Source/CustomLaunchUI/Private/Marshalling/GenericProfileMarshaller.cpp
+++ b/Engine/Plugins/Developer/CustomLaunchUI/Source/CustomLaunchUI/Private/Marshalling/GenericProfileMarshaller.cpp
@@ -55,6 +55,11 @@ namespace CustomLaunchUI
 			return (bShouldCook && ContentScheme != EContentScheme::ZenPakStreaming && ContentScheme != EContentScheme::CookOnTheFly);
 		};
 
+		ForContent = [this]()
+		{
+			return (ContentScheme != EContentScheme::ZenPakStreaming);
+		};
+
 		ForCode = [this]()
 		{
 			return GetBuild();
@@ -243,7 +248,7 @@ namespace CustomLaunchUI
 				.GetValue = [this]()			{ return Profile->IsImportingZenSnapshot(); },
 				.SetValue = [this](bool bValue)	{ Profile->SetImportingZenSnapshot(bValue); },
 				.GetDefaultValue = [this]()		{ return DefaultProfile->IsImportingZenSnapshot(); },
-				.IsVisible = ForCooked,
+				.IsVisible = ForContent,
 			}
 		);
 	}
@@ -262,14 +267,46 @@ namespace CustomLaunchUI
 
 	void FGenericProfileMarshaller::AddIncrementalCookProperty( FLaunchProfileTreeNode& HeadingNode )
 	{
-		HeadingNode.AddBoolean( LOCTEXT("IncrementalCookLabel", "Incremental Cook"),
+		auto GetDisplayName = []( ELauncherProfileIncrementalCookMode::Type Mode)
+		{
+			switch (Mode)
 			{
-				.GetValue = [this]()			{ return Profile->IsCookingIncrementally() || Profile->IsImportingZenSnapshot(); }, // show it as ticked when snapshot import is enabled for now to indicate that it will be doing something incremental. eventually this is likely to become a separate field
-				.SetValue = [this](bool bValue)	{ Profile->SetIncrementalCooking(bValue); },
-				.GetDefaultValue = [this]()		{ return DefaultProfile->IsCookingIncrementally(); },
-				.IsVisible = ForCooked,
-				.IsEnabled = [this]()			{ return !Profile->IsImportingZenSnapshot(); }, // zen snapshot import & incremental cooking can't be used together. '-iterativecook' is the 'old' system, and zen import will enable the new '-cookincremental' & the two flags don't play nicely together
+				case ELauncherProfileIncrementalCookMode::None:                    return LOCTEXT("IncrementalCookNone",                 "None");
+				case ELauncherProfileIncrementalCookMode::ModifiedOnly:            return LOCTEXT("IncrementalCookModified",             "Modified Only (legacy)");
+				case ELauncherProfileIncrementalCookMode::ModifiedAndDependencies: return LOCTEXT("IncrementalCookModifiedDependencies", "Modified & Dependencies (recommended)");
 			}
+			return FText::GetEmpty();
+		};
+
+		auto GetToolTip = []( ELauncherProfileIncrementalCookMode::Type Mode)
+		{
+			switch (Mode)
+			{
+				case ELauncherProfileIncrementalCookMode::None:                    return LOCTEXT("IncrementalCookNoneTip",                 "This will try to cook everything");
+				case ELauncherProfileIncrementalCookMode::ModifiedOnly:            return LOCTEXT("IncrementalCookModifiedTip",             "This will only try to cook any modified assets but won't try to cook anything that depends on these assets. This is the old, legacy option and is faster but unreliable");
+				case ELauncherProfileIncrementalCookMode::ModifiedAndDependencies: return LOCTEXT("IncrementalCookModifiedDependenciesTip", "This will try to cook any modified assets and those assets that depend on them. This is the new method and is much more reliable but is a little slower");
+			}
+			return FText::GetEmpty();
+		};
+
+		HeadingNode.AddWidget( LOCTEXT("IncrementalCookLabel", "Incremental Cook"),
+			{
+				.IsDefault = [this]()		{ return GetIncrementalCookMode() == DefaultProfile->GetIncrementalCookMode(); },
+				.SetToDefault = [this]()	{ SetIncrementalCookMode(DefaultProfile->GetIncrementalCookMode()); },
+				.IsVisible = ForEnabledCooked,
+			},
+			SNew(SCustomLaunchCombo<ELauncherProfileIncrementalCookMode::Type>)
+			.OnSelectionChanged(this, &FGenericProfileMarshaller::SetIncrementalCookMode)
+			.SelectedItem(this, &FGenericProfileMarshaller::GetIncrementalCookMode)
+			.GetDisplayName_Lambda(GetDisplayName)
+			.GetItemToolTip_Lambda(GetToolTip)
+			.Items( TArray<ELauncherProfileIncrementalCookMode::Type>(
+				{
+					ELauncherProfileIncrementalCookMode::None,
+					ELauncherProfileIncrementalCookMode::ModifiedOnly,
+					ELauncherProfileIncrementalCookMode::ModifiedAndDependencies,
+				})
+			)
 		);
 	}
 
@@ -914,6 +951,19 @@ namespace CustomLaunchUI
 	}
 
 
+	void FGenericProfileMarshaller::SetIncrementalCookMode( ELauncherProfileIncrementalCookMode::Type Mode )
+	{
+		Profile->SetIncrementalCookMode(Mode);
+		Profile->SetUnversionedCooking(Mode != ELauncherProfileIncrementalCookMode::ModifiedOnly); // should always use unversioned, except the legacy -iterativecook because it does not handle invalidation due to c++ changes.
+		OnPropertyChanged();
+	}
+
+	ELauncherProfileIncrementalCookMode::Type FGenericProfileMarshaller::GetIncrementalCookMode() const
+	{
+		return Profile->GetIncrementalCookMode();
+	}
+
+
 
 	void FGenericProfileMarshaller::SetMapsToCook(TArray<FString> MapsToCook)
 	{
diff --git a/Engine/Plugins/Developer/CustomLaunchUI/Source/CustomLaunchUI/Private/Model/CustomLaunchUIModel.cpp b/Engine/Plugins/Developer/CustomLaunchUI/Source/CustomLaunchUI/Private/Model/CustomLaunchUIModel.cpp
--- a/Engine/Plugins/Developer/CustomLaunchUI/Source/CustomLaunchUI/Private/Model/CustomLaunchUIModel.cpp
+++ b/Engine/Plugins/Developer/CustomLaunchUI/Source/CustomLaunchUI/Private/Model/CustomLaunchUIModel.cpp
@@ -717,7 +717,7 @@ namespace CustomLaunchUI
 	{
 		switch (ContentScheme)
 		{
-			case EContentScheme::PakFiles:           return LOCTEXT("ContentSchemeTipPakFiles","Store cooked game content in one ore more large Pak Files");
+			case EContentScheme::PakFiles:           return LOCTEXT("ContentSchemeTipPakFiles","Store cooked game content in one or more large Pak Files");
 			case EContentScheme::ZenStreaming:       return LOCTEXT("ContentSchemeTipZenStreaming","Stream cooked game content from Zen Server");
 			case EContentScheme::ZenPakStreaming:    return LOCTEXT("ContentSchemeTipZenPakStreaming","Stream an existing Pak Files build via a Zen");
 			case EContentScheme::DevelopmentPackage: return LOCTEXT("ContentSchemeTipDevPackage","Package cooked game content into a single installable package file for development purposes, where available");
diff --git a/Engine/Plugins/Developer/CustomLaunchUI/Source/CustomLaunchUI/Private/Styling/CustomLaunchUIStyle.cpp b/Engine/Plugins/Developer/CustomLaunchUI/Source/CustomLaunchUI/Private/Styling/CustomLaunchUIStyle.cpp
--- a/Engine/Plugins/Developer/CustomLaunchUI/Source/CustomLaunchUI/Private/Styling/CustomLaunchUIStyle.cpp
+++ b/Engine/Plugins/Developer/CustomLaunchUI/Source/CustomLaunchUI/Private/Styling/CustomLaunchUIStyle.cpp
@@ -81,6 +81,7 @@ void FCustomLaunchUIStyle::Initialize()
 	StyleSet->Set("Icons.Task.Cleanup", new CORE_IMAGE_BRUSH_SVG("../Editor/Slate/Starship/GraphEditors/CleanUp", Icon16x16));
 	StyleSet->Set("Icons.Task.Launch",	new CORE_IMAGE_BRUSH_SVG("Starship/Common/ProjectLauncher", Icon16x16));
 	StyleSet->Set("Icons.Task.Build", new CORE_IMAGE_BRUSH_SVG("../Editor/Slate/Starship/Common/Adjust", Icon16x16));
+	StyleSet->Set("Icons.Task.Zen", new IMAGE_BRUSH_SVG("Zen_16", Icon16x16));
 	StyleSet->Set("Icons.Task.Cook", new CORE_IMAGE_BRUSH_SVG( "../Editor/Slate/Starship/Common/CookContent", Icon16x16 ) );
 	StyleSet->Set("Icons.Task.Deploy", new CORE_IMAGE_BRUSH_SVG( "../Editor/Slate/Starship/Common/DeviceManager", Icon16x16 ) );
 	StyleSet->Set("Icons.Task.Archive", new CORE_IMAGE_BRUSH_SVG("../Editor/Slate/Starship/Common/SaveCurrent", Icon16x16));
@@ -139,6 +140,10 @@ const FSlateBrush* FCustomLaunchUIStyle::GetBrushForTask(ILauncherTaskPtr Task)
 		{
 			return GetBrush("Icons.Task.Build");
 		}
+		else if (TaskName.Contains(TEXT("Snapshot")))
+		{
+			return GetBrush("Icons.Task.Zen");
+		}
 		else if (TaskName.Contains(TEXT("Cook")))
 		{
 			return GetBrush("Icons.Task.Cook");
diff --git a/Engine/Plugins/Developer/CustomLaunchUI/Source/CustomLaunchUI/Public/Marshalling/GenericProfileMarshaller.h b/Engine/Plugins/Developer/CustomLaunchUI/Source/CustomLaunchUI/Public/Marshalling/GenericProfileMarshaller.h
--- a/Engine/Plugins/Developer/CustomLaunchUI/Source/CustomLaunchUI/Public/Marshalling/GenericProfileMarshaller.h
+++ b/Engine/Plugins/Developer/CustomLaunchUI/Source/CustomLaunchUI/Public/Marshalling/GenericProfileMarshaller.h
@@ -75,6 +75,7 @@ namespace CustomLaunchUI
 		FLaunchProfileTreeNode::FGetBool ForZenWS;
 		FLaunchProfileTreeNode::FGetBool ForCooked;
 		FLaunchProfileTreeNode::FGetBool ForEnabledCooked;
+		FLaunchProfileTreeNode::FGetBool ForContent;
 		FLaunchProfileTreeNode::FGetBool ForCode;
 		FLaunchProfileTreeNode::FGetBool ForDeployment;
 		FLaunchProfileTreeNode::FGetBool ForRun;
@@ -122,6 +123,9 @@ namespace CustomLaunchUI
 		void SetCook( bool bCook );
 		bool GetCook( ILauncherProfilePtr InProfile = nullptr ) const;
 
+		void SetIncrementalCookMode( ELauncherProfileIncrementalCookMode::Type Mode );
+		ELauncherProfileIncrementalCookMode::Type GetIncrementalCookMode() const;
+
 		void SetMapsToCook(TArray<FString> MapsToCook);
 		TArray<FString> GetMapsToCook() const;
 		EMapOption GetMapOption() const;
diff --git a/Engine/Plugins/Editor/MobileLauncherProfileWizard/Source/MobileLauncherProfileWizard/Private/AndroidProfileWizard.cpp b/Engine/Plugins/Editor/MobileLauncherProfileWizard/Source/MobileLauncherProfileWizard/Private/AndroidProfileWizard.cpp
--- a/Engine/Plugins/Editor/MobileLauncherProfileWizard/Source/MobileLauncherProfileWizard/Private/AndroidProfileWizard.cpp
+++ b/Engine/Plugins/Editor/MobileLauncherProfileWizard/Source/MobileLauncherProfileWizard/Private/AndroidProfileWizard.cpp
@@ -49,7 +49,7 @@ static void SetupAndroidAppProfile(ILauncherProfileRef& AppProfile, const FProfi
 	// Release settings
 	AppProfile->SetCreateReleaseVersion(true);
 	AppProfile->SetCreateReleaseVersionName(AndroidProfileConstants::AppReleaseName);
-	AppProfile->SetIncrementalCooking(false);
+	AppProfile->SetIncrementalCookMode(ELauncherProfileIncrementalCookMode::None);
 	AppProfile->SetCompressed(false);
 	AppProfile->SetDeployWithUnrealPak(true);
 
@@ -107,7 +107,7 @@ static void SetupAndroidDLCProfile(ILauncherProfileRef& DLCProfile, const FProfi
 	DLCProfile->SetDLCName(AndroidProfileConstants::DLCName);
 	DLCProfile->SetDLCIncludeEngineContent(true);
 		
-	DLCProfile->SetIncrementalCooking(false);
+	DLCProfile->SetIncrementalCookMode(ELauncherProfileIncrementalCookMode::None);
 	DLCProfile->SetCompressed(false);
 	DLCProfile->SetDeployWithUnrealPak(true);
 
diff --git a/Engine/Plugins/Editor/MobileLauncherProfileWizard/Source/MobileLauncherProfileWizard/Private/IOSProfileWizard.cpp b/Engine/Plugins/Editor/MobileLauncherProfileWizard/Source/MobileLauncherProfileWizard/Private/IOSProfileWizard.cpp
--- a/Engine/Plugins/Editor/MobileLauncherProfileWizard/Source/MobileLauncherProfileWizard/Private/IOSProfileWizard.cpp
+++ b/Engine/Plugins/Editor/MobileLauncherProfileWizard/Source/MobileLauncherProfileWizard/Private/IOSProfileWizard.cpp
@@ -49,7 +49,7 @@ static void SetupIOSAppProfile(ILauncherProfileRef& AppProfile, const FProfilePa
 	// Release settings
 	AppProfile->SetCreateReleaseVersion(true);
 	AppProfile->SetCreateReleaseVersionName(IOSProfileConstants::AppReleaseName);
-	AppProfile->SetIncrementalCooking(false);
+	AppProfile->SetIncrementalCookMode(ELauncherProfileIncrementalCookMode::None);
 	AppProfile->SetCompressed(false);
 	AppProfile->SetDeployWithUnrealPak(true);
 
@@ -104,7 +104,7 @@ static void SetupIOSDLCProfile(ILauncherProfileRef& DLCProfile, const FProfilePa
 	DLCProfile->SetDLCName(IOSProfileConstants::DLCName);
 	DLCProfile->SetDLCIncludeEngineContent(true);
 		
-	DLCProfile->SetIncrementalCooking(false);
+	DLCProfile->SetIncrementalCookMode(ELauncherProfileIncrementalCookMode::None);
 	DLCProfile->SetCompressed(false);
 	DLCProfile->SetDeployWithUnrealPak(true);
 
diff --git a/Engine/Source/Developer/LauncherServices/Private/Launcher/LauncherWorker.cpp b/Engine/Source/Developer/LauncherServices/Private/Launcher/LauncherWorker.cpp
--- a/Engine/Source/Developer/LauncherServices/Private/Launcher/LauncherWorker.cpp
+++ b/Engine/Source/Developer/LauncherServices/Private/Launcher/LauncherWorker.cpp
@@ -682,6 +682,20 @@ FString FLauncherWorker::CreateUATCommand( const ILauncherProfileRef& InProfile,
 		// @todo: server
 	}
 
+	// snapshot import
+	if (InProfile->IsImportingZenSnapshot())
+	{
+		UATCommand += TEXT(" -snapshot");
+
+		FCommandDesc Desc;
+		FText Command = FText::Format(LOCTEXT("LauncherSnapshotDesc", "Import Zen snapshot for {0}"), FText::FromString(Platforms.RightChop(1)));
+		Desc.Name = "Snapshot Import Task";
+		Desc.Desc = Command.ToString();
+		Desc.EndText = TEXT("********** SNAPSHOT IMPORT COMPLETED **********");
+		OutCommands.Add(Desc);
+		CommandStart = TEXT("********** SNAPSHOT IMPORT STARTED **********");
+	}
+
 	// cook
 	switch(InProfile->GetCookMode())
 	{
@@ -742,10 +756,6 @@ FString FLauncherWorker::CreateUATCommand( const ILauncherProfileRef& InProfile,
 			{
 				UATCommand += TEXT(" -zenstore");
 			}
-			if (InProfile->IsImportingZenSnapshot())
-			{
-				UATCommand += TEXT(" -snapshot");
-			}
 
 			if (FDerivedDataCacheInterface* DDC = TryGetDerivedDataCache())
 			{
@@ -856,10 +866,6 @@ FString FLauncherWorker::CreateUATCommand( const ILauncherProfileRef& InProfile,
 			{
 				UATCommand += TEXT(" -zenstore");
 			}
-			if (InProfile->IsImportingZenSnapshot())
-			{
-				UATCommand += TEXT(" -snapshot");
-			}
 
 			TStringBuilder<64> InheritableCookOptions;
 			FCommandLine::BuildSubprocessCommandLine(ECommandLineArgumentFlags::CommandletContext, InheritableCookOptions);
@@ -925,10 +931,6 @@ FString FLauncherWorker::CreateUATCommand( const ILauncherProfileRef& InProfile,
 		{
 			UATCommand += TEXT(" -zenstore");
 		}
-		if (InProfile->IsImportingZenSnapshot())
-		{
-			UATCommand += TEXT(" -snapshot");
-		}
 		break;
 	case ELauncherProfileCookModes::ByTheBookInEditor:
 		UATCommand += MapList;
@@ -938,10 +940,6 @@ FString FLauncherWorker::CreateUATCommand( const ILauncherProfileRef& InProfile,
 		{
 			UATCommand += TEXT(" -zenstore");
 		}
-		if (InProfile->IsImportingZenSnapshot())
-		{
-			UATCommand += TEXT(" -snapshot");
-		}
 		if (InProfile->IsPackingWithUnrealPak())
 		{
 			UATCommand += TEXT(" -pak");
@@ -969,9 +967,19 @@ FString FLauncherWorker::CreateUATCommand( const ILauncherProfileRef& InProfile,
 		UATCommand += TEXT(" -distribution");
 	}
 
-	if (InProfile->IsCookingIncrementally() && !InProfile->IsImportingZenSnapshot()) // don't set -iterativecooking when importing a zen snapshot because it sets -cookincremental (the new system) instead.
+	switch (InProfile->GetIncrementalCookMode())
 	{
-		UATCommand += TEXT(" -iterativecooking");
+		case ELauncherProfileIncrementalCookMode::ModifiedOnly:
+		{
+			UATCommand += TEXT(" -iterativecooking");
+		}
+		break;
+
+		case ELauncherProfileIncrementalCookMode::ModifiedAndDependencies:
+		{
+			UATCommand += TEXT(" -cookincremental");
+		}
+		break;
 	}
 
 	if ( InProfile->IsIterateSharedCookedBuild() )
diff --git a/Engine/Source/Developer/LauncherServices/Private/Profiles/LauncherProfile.h b/Engine/Source/Developer/LauncherServices/Private/Profiles/LauncherProfile.h
--- a/Engine/Source/Developer/LauncherServices/Private/Profiles/LauncherProfile.h
+++ b/Engine/Source/Developer/LauncherServices/Private/Profiles/LauncherProfile.h
@@ -59,6 +59,7 @@ enum ELauncherVersion
 	LAUNCHERSERVICES_ADDZENOPTIONS = 37,
 	LAUNCHERSERVICES_ADDBUILDARCHITECTURE = 38,
 	LAUNCHERSERVICES_ADDCUSTOMPROPERTIES = 39,
+	LAUNCHERSERVICES_INCREMENTALCOOKMODE = 40,
 	//ADD NEW STUFF HERE
 
 
@@ -815,11 +816,17 @@ public:
 	{
 		if ( CookMode != ELauncherProfileCookModes::DoNotCook )
 		{
-			return CookIncremental;
+			return (IncrementalCookMode == ELauncherProfileIncrementalCookMode::ModifiedOnly);
 		}
 		return false;
 	}
 
+	virtual ELauncherProfileIncrementalCookMode::Type GetIncrementalCookMode() const override
+	{
+		return IncrementalCookMode;
+	}
+
+
 	virtual bool IsIterateSharedCookedBuild() const override
 	{
 		if ( CookMode != ELauncherProfileCookModes::DoNotCook )
@@ -968,6 +975,7 @@ public:
 
 		// IMPORTANT: make sure to bump LAUNCHERSERVICES_PROFILEVERSION when modifying this!
 		bool BuildGame = false;
+		bool CookIncremental = false;
 		Archive << Id
 				<< Name
 				<< Description
@@ -1138,6 +1146,15 @@ public:
 			Archive << CustomBoolProperties;
 		}
 
+		if (Version >= LAUNCHERSERVICES_INCREMENTALCOOKMODE)
+		{
+			Archive << IncrementalCookMode;
+		}
+		else if (Archive.IsLoading())
+		{
+			IncrementalCookMode = CookIncremental ? ELauncherProfileIncrementalCookMode::ModifiedOnly : ELauncherProfileIncrementalCookMode::None;
+		}
+
 		DefaultLaunchRole->Serialize(Archive);
 
 		// serialize launch roles
@@ -1204,7 +1221,6 @@ public:
 		Writer.WriteValue("ProjectSpecified", ProjectSpecified);
 		Writer.WriteValue("ShareableProjectPath", ShareableProjectPath);
 		Writer.WriteValue("CookConfiguration", (int32)CookConfiguration);
-		Writer.WriteValue("CookIncremental", CookIncremental);
 		Writer.WriteValue("CookOptions", CookOptions);
 		Writer.WriteValue("CookMode", CookMode);
 		Writer.WriteValue("CookUnversioned", CookUnversioned);
@@ -1292,7 +1308,7 @@ public:
 		Writer.WriteValue("EditorArchitectures", EditorArchitectures);
 		Writer.WriteValue("CustomStringProperties", CustomStringProperties);
 		Writer.WriteValue("CustomBoolProperties", CustomBoolProperties);
-
+		Writer.WriteValue("IncrementalCookMode", IncrementalCookMode);
 
 		// serialize the default launch role
 		DefaultLaunchRole->Save(Writer, TEXT("DefaultRole"));
@@ -1844,6 +1860,8 @@ public:
 			return false;
 		}
 
+		bool CookIncremental;
+
 		FGuid::Parse(Object.GetStringField(TEXT("Id")), Id);
 		Name = Object.GetStringField(TEXT("Name"));
 		Description = Object.GetStringField(TEXT("Description"));
@@ -2050,6 +2068,15 @@ public:
 			}
 		}
 
+		if (Version >= LAUNCHERSERVICES_INCREMENTALCOOKMODE)
+		{
+			IncrementalCookMode = (TEnumAsByte<ELauncherProfileIncrementalCookMode::Type>)((int32)Object.GetNumberField(TEXT("IncrementalCookMode")));
+		}
+		else
+		{
+			IncrementalCookMode = CookIncremental ? ELauncherProfileIncrementalCookMode::ModifiedOnly : ELauncherProfileIncrementalCookMode::None;
+		}
+
 
 		// load the default launch role
 		TSharedPtr<FJsonObject> Role = Object.GetObjectField(TEXT("DefaultRole"));
@@ -2133,7 +2160,7 @@ public:
 		CookConfiguration = FApp::GetBuildConfiguration();
 		CookMode = ELauncherProfileCookModes::OnTheFly;
 		CookOptions = FString();
-		CookIncremental = false;
+		IncrementalCookMode = ELauncherProfileIncrementalCookMode::None;
 		IterateSharedCookedBuild=false;
 		CookUnversioned = true;
 		Compressed = true;
@@ -2424,14 +2451,20 @@ public:
 
 	virtual void SetIncrementalCooking( bool Incremental ) override
 	{
-		if (CookIncremental != Incremental)
+		SetIncrementalCookMode( Incremental ? ELauncherProfileIncrementalCookMode::ModifiedOnly : ELauncherProfileIncrementalCookMode::None );
+	}
+
+	virtual void SetIncrementalCookMode( ELauncherProfileIncrementalCookMode::Type Mode ) override
+	{
+		if (IncrementalCookMode != Mode)
 		{
-			CookIncremental = Incremental;
+			IncrementalCookMode = Mode;
 
 			Validate();
 		}
 	}
 
+
 	virtual void SetIterateSharedCookedBuild( bool SharedCookedBuild ) override
 	{
 		if (IterateSharedCookedBuild != SharedCookedBuild)
@@ -2732,11 +2765,7 @@ public:
 
 	virtual bool IsImportingZenSnapshot() const override
 	{
-		if (CookMode != ELauncherProfileCookModes::DoNotCook)
-		{
-			return bImportingZenSnapshot;
-		}
-		return false;
+		return bImportingZenSnapshot;
 	}
 
 	virtual void SetUseZenPakStreaming( bool UseZenPakStreaming ) override
@@ -2949,7 +2978,7 @@ protected:
 			ValidationErrors.Add(ELauncherProfileValidationErrors::ZenPakStreamingRequiresDeployAndLaunch);
 		}
 
-		if (CookUnversioned && CookIncremental && ((CookMode == ELauncherProfileCookModes::ByTheBook) || (CookMode == ELauncherProfileCookModes::ByTheBookInEditor)))
+		if (CookUnversioned && (IncrementalCookMode == ELauncherProfileIncrementalCookMode::ModifiedOnly) && ((CookMode == ELauncherProfileCookModes::ByTheBook) || (CookMode == ELauncherProfileCookModes::ByTheBookInEditor)))
 		{
 			ValidationErrors.Add(ELauncherProfileValidationErrors::UnversionedAndIncremental);
 		}
@@ -3291,8 +3320,9 @@ private:
 	bool EncryptIniFiles;
 	// is this build for distribution
 	bool ForDistribution;
-	// Holds a flag indicating whether only modified content should be cooked.
-	bool CookIncremental;
+	
+	// Holds the incremental cook mode indicating whether only modified content should be cooked.
+	TEnumAsByte<ELauncherProfileIncrementalCookMode::Type> IncrementalCookMode;
 
 	// hold a flag indicating if we want to iterate on the shared cooked build
 	bool IterateSharedCookedBuild;
diff --git a/Engine/Source/Developer/LauncherServices/Public/ILauncherProfile.h b/Engine/Source/Developer/LauncherServices/Public/ILauncherProfile.h
--- a/Engine/Source/Developer/LauncherServices/Public/ILauncherProfile.h
+++ b/Engine/Source/Developer/LauncherServices/Public/ILauncherProfile.h
@@ -138,6 +138,26 @@ namespace ELauncherProfilePackagingModes
 }
 
 
+namespace ELauncherProfileIncrementalCookMode
+{
+	/**
+	 * Enumerates modes in which the launcher handles incremental cooking
+	 */
+	enum Type
+	{
+		/** Do not cook incrementally (default). */
+		None,
+
+		/** -iterativecooking : Only recook specifically modified assets without recooking dependencies. This is the legacy method and is must less robust. */
+		ModifiedOnly,
+
+		/** -cookincremental : Recook modified assets and any dependencies. This is the new method and is currently slower but far more robust. */
+		ModifiedAndDependencies,
+	};
+}
+
+
+
 namespace ELauncherProfileValidationErrors
 {
 	/**
@@ -792,13 +812,22 @@ public:
 	virtual bool IsBuildingUAT() const = 0;
 
 	/**
-	 * Checks whether incremental cooking is enabled.
+	 * Checks whether legacy incremental cooking is enabled.
 	 *
 	 * @return true if cooking incrementally, false otherwise.
 	 * @see SetIncrementalCooking
 	 */
+	UE_DEPRECATED(5.6, "Use GetIncrementalCookMode instead")
 	virtual bool IsCookingIncrementally( ) const = 0;
 
+	/**
+	 * Get the current incremental cook mode
+	 * 
+	 * @return incremental cook mode
+	 * @see SetIncrementalCookMode
+	 */ 
+	virtual ELauncherProfileIncrementalCookMode::Type GetIncrementalCookMode() const = 0;
+
 
 	virtual bool IsIterateSharedCookedBuild() const =0;
 
@@ -1201,13 +1230,23 @@ public:
 	virtual void SetHideFileServerWindow( bool Hide ) = 0;
 
 	/**
-	 * Sets incremental cooking.
+	 * Sets legacy incremental cooking.
 	 *
 	 * @param Incremental Whether cooking should be incremental.
 	 * @see IsCookingIncrementally
 	 */
+	UE_DEPRECATED(5.6, "Use SetIncrementalCookMode instead")
 	virtual void SetIncrementalCooking( bool Incremental ) = 0;
 
+	/**
+	 * Set the incremental cook mode.
+	 * 
+	 * @param Mode the incremental cooking mode to use
+	 * @see GetIncrementalCookMode
+	 */
+	virtual void SetIncrementalCookMode( ELauncherProfileIncrementalCookMode::Type Mode ) = 0;
+
+
 	virtual void SetIterateSharedCookedBuild( bool IterateSharedCookedBuild ) = 0;
 
 	/**
diff --git a/Engine/Source/Developer/ProjectLauncher/Private/Widgets/Cook/SProjectLauncherCookByTheBookSettings.cpp b/Engine/Source/Developer/ProjectLauncher/Private/Widgets/Cook/SProjectLauncherCookByTheBookSettings.cpp
--- a/Engine/Source/Developer/ProjectLauncher/Private/Widgets/Cook/SProjectLauncherCookByTheBookSettings.cpp
+++ b/Engine/Source/Developer/ProjectLauncher/Private/Widgets/Cook/SProjectLauncherCookByTheBookSettings.cpp
@@ -1305,7 +1305,8 @@ void SProjectLauncherCookByTheBookSettings::HandleIncrementalCheckBoxCheckStateC
 
 	if (SelectedProfile.IsValid())
 	{
-		SelectedProfile->SetIncrementalCooking(NewState == ECheckBoxState::Checked);
+		bool bIncremental = (NewState == ECheckBoxState::Checked);
+		SelectedProfile->SetIncrementalCookMode( bIncremental ? ELauncherProfileIncrementalCookMode::ModifiedOnly : ELauncherProfileIncrementalCookMode::None);
 	}
 }
 
@@ -1316,7 +1317,7 @@ ECheckBoxState SProjectLauncherCookByTheBookSettings::HandleIncrementalCheckBoxI
 
 	if (SelectedProfile.IsValid())
 	{
-		if (SelectedProfile->IsCookingIncrementally())
+		if (SelectedProfile->GetIncrementalCookMode() != ELauncherProfileIncrementalCookMode::None)
 		{
 			return ECheckBoxState::Checked;
 		}
diff --git a/Engine/Source/Developer/ProjectLauncher/Private/Widgets/Cook/SProjectLauncherCookOnTheFlySettings.cpp b/Engine/Source/Developer/ProjectLauncher/Private/Widgets/Cook/SProjectLauncherCookOnTheFlySettings.cpp
--- a/Engine/Source/Developer/ProjectLauncher/Private/Widgets/Cook/SProjectLauncherCookOnTheFlySettings.cpp
+++ b/Engine/Source/Developer/ProjectLauncher/Private/Widgets/Cook/SProjectLauncherCookOnTheFlySettings.cpp
@@ -104,7 +104,8 @@ void SProjectLauncherCookOnTheFlySettings::HandleIncrementalCheckBoxCheckStateCh
 
 	if (SelectedProfile.IsValid())
 	{
-		SelectedProfile->SetIncrementalCooking(NewState == ECheckBoxState::Checked);
+		bool bIncremental = (NewState == ECheckBoxState::Checked);
+		SelectedProfile->SetIncrementalCookMode( bIncremental ? ELauncherProfileIncrementalCookMode::ModifiedOnly : ELauncherProfileIncrementalCookMode::None);
 	}
 }
 
@@ -115,7 +116,7 @@ ECheckBoxState SProjectLauncherCookOnTheFlySettings::HandleIncrementalCheckBoxIs
 
 	if (SelectedProfile.IsValid())
 	{
-		if (SelectedProfile->IsCookingIncrementally())
+		if (SelectedProfile->GetIncrementalCookMode() != ELauncherProfileIncrementalCookMode::None)
 		{
 			return ECheckBoxState::Checked;
 		}
diff --git a/Engine/Source/Editor/UnrealEd/Private/PlayLevelViaLauncher.cpp b/Engine/Source/Editor/UnrealEd/Private/PlayLevelViaLauncher.cpp
--- a/Engine/Source/Editor/UnrealEd/Private/PlayLevelViaLauncher.cpp
+++ b/Engine/Source/Editor/UnrealEd/Private/PlayLevelViaLauncher.cpp
@@ -293,7 +293,7 @@ void UEditorEngine::StartPlayUsingLauncherSession(FRequestPlaySessionParams& InR
 	}
 	LauncherProfile->SetCookMode(CurrentLauncherCookMode);
 	LauncherProfile->SetUnversionedCooking(!bIncrementalCooking); // Unversioned cooking is not allowed with incremental cooking
-	LauncherProfile->SetIncrementalCooking(bIncrementalCooking);
+	LauncherProfile->SetIncrementalCookMode(bIncrementalCooking ? ELauncherProfileIncrementalCookMode::ModifiedOnly : ELauncherProfileIncrementalCookMode::None);
 	SetCookOption(TEXTVIEW("-IgnoreIniSettingsOutOfDate"), bIncrementalCooking && CookerSettings.bIgnoreIniSettingsOutOfDateForIteration);
 	SetCookOption(TEXTVIEW("-IgnoreScriptPackagesOutOfDate"), bIncrementalCooking && CookerSettings.bIgnoreScriptPackagesOutOfDateForIteration);
 	SetCookOption(TEXTVIEW("-IterateSharedCookedbuild"), bIncrementalCooking && ExperimentalSettings.bSharedCookedBuilds);
