Fix issues caused by the -Multiprocess argument being added to invocations of the cook commandlet from the editor.  Because this argument causes the config system to disable writes, it blocks functionality that the cooker relies on.

diff --git a/Engine/Plugins/MovieScene/MovieRenderPipeline/Source/MovieRenderPipelineEditor/Private/MoviePipelineNewProcessExecutor.cpp b/Engine/Plugins/MovieScene/MovieRenderPipeline/Source/MovieRenderPipelineEditor/Private/MoviePipelineNewProcessExecutor.cpp
--- a/Engine/Plugins/MovieScene/MovieRenderPipeline/Source/MovieRenderPipelineEditor/Private/MoviePipelineNewProcessExecutor.cpp
+++ b/Engine/Plugins/MovieScene/MovieRenderPipeline/Source/MovieRenderPipelineEditor/Private/MoviePipelineNewProcessExecutor.cpp
@@ -204,7 +204,7 @@ void UMoviePipelineNewProcessExecutor::Execute_Implementation(UMoviePipelineQueu
 	NamedArguments.Add(TEXT("PlayWorld"), TEXT("MoviePipelineEntryMap")); // Boot up on an empty map, the executor will immediately transition to the correct one.
 	NamedArguments.Add(TEXT("UnrealURL"), UnrealURLParams); // Pass the command line arguments for this job
 	TStringBuilder<64> SubprocessCommandLine;
-	FCommandLine::BuildSubprocessCommandLine(ECommandLineArgumentFlags::GameContexts, SubprocessCommandLine);
+	FCommandLine::BuildSubprocessCommandLine(ECommandLineArgumentFlags::GameContexts, false /*bOnlyInherited*/, SubprocessCommandLine);
 	NamedArguments.Add(TEXT("SubprocessCommandLine"), *SubprocessCommandLine);
 	NamedArguments.Add(TEXT("CommandLineParams"), CommandLineArgs);
 	NamedArguments.Add(TEXT("MoviePipelineArgs"), MoviePipelineArgs);
diff --git a/Engine/Source/Developer/LauncherServices/Private/Launcher/LauncherWorker.cpp b/Engine/Source/Developer/LauncherServices/Private/Launcher/LauncherWorker.cpp
--- a/Engine/Source/Developer/LauncherServices/Private/Launcher/LauncherWorker.cpp
+++ b/Engine/Source/Developer/LauncherServices/Private/Launcher/LauncherWorker.cpp
@@ -618,7 +618,9 @@ FString FLauncherWorker::CreateUATCommand( const ILauncherProfileRef& InProfile,
 	FString AdditionalCommandLine;
 	TStringBuilder<64> InheritableGameOptions;
 	ECommandLineArgumentFlags CommandLineContextFlags = ECommandLineArgumentFlags::ClientContext | ECommandLineArgumentFlags::ServerContext;
-	FCommandLine::BuildSubprocessCommandLine(CommandLineContextFlags, InheritableGameOptions);
+	const bool bOnlyInherited = true; // Only getting inherited, not explicitly set subprocess commandline args due to issues when passing
+									  // the -Multiprocess explicitly set commandline argument which prevents configs from saving/writing
+	FCommandLine::BuildSubprocessCommandLine(CommandLineContextFlags, bOnlyInherited, InheritableGameOptions);
 	FString EscapedInheritableGameOptions(InheritableGameOptions);
 	EscapedInheritableGameOptions.ReplaceInline(TEXT("\""), TEXT("\\\""));
 	FString AddCmdLine = FString::Printf(TEXT("%s%s %s %s%s"),
@@ -717,7 +719,7 @@ FString FLauncherWorker::CreateUATCommand( const ILauncherProfileRef& InProfile,
 			}
 
 			TStringBuilder<64> InheritableCookOptions;
-			FCommandLine::BuildSubprocessCommandLine(ECommandLineArgumentFlags::CommandletContext, InheritableCookOptions);
+			FCommandLine::BuildSubprocessCommandLine(ECommandLineArgumentFlags::CommandletContext, bOnlyInherited, InheritableCookOptions);
 			FString AdditionalOptions = InProfile->GetCookOptions();
 			if (!AdditionalOptions.IsEmpty() || (InheritableCookOptions.Len() > 0))
 			{
@@ -868,7 +870,7 @@ FString FLauncherWorker::CreateUATCommand( const ILauncherProfileRef& InProfile,
 			}
 
 			TStringBuilder<64> InheritableCookOptions;
-			FCommandLine::BuildSubprocessCommandLine(ECommandLineArgumentFlags::CommandletContext, InheritableCookOptions);
+			FCommandLine::BuildSubprocessCommandLine(ECommandLineArgumentFlags::CommandletContext, bOnlyInherited, InheritableCookOptions);
 
 			FString AdditionalOptions = InProfile->GetCookOptions();
 			if (!AdditionalOptions.IsEmpty() || (InheritableCookOptions.Len() > 0))
diff --git a/Engine/Source/Editor/MovieSceneCaptureDialog/Private/MovieSceneCaptureDialogModule.cpp b/Engine/Source/Editor/MovieSceneCaptureDialog/Private/MovieSceneCaptureDialogModule.cpp
--- a/Engine/Source/Editor/MovieSceneCaptureDialog/Private/MovieSceneCaptureDialogModule.cpp
+++ b/Engine/Source/Editor/MovieSceneCaptureDialog/Private/MovieSceneCaptureDialogModule.cpp
@@ -735,7 +735,7 @@ void FNewProcessCapture::Start()
 	EditorCommandLine += FString::Printf(TEXT(" -messaging -SessionName=\"%s\""), MovieCaptureSessionName);
 
 	TStringBuilder<64> SubprocessCommandLine;
-	FCommandLine::BuildSubprocessCommandLine(ECommandLineArgumentFlags::GameContexts, SubprocessCommandLine);
+	FCommandLine::BuildSubprocessCommandLine(ECommandLineArgumentFlags::GameContexts, false /*bOnlyInherited*/, SubprocessCommandLine);
 	FString Params;
 	if (FPaths::IsProjectFilePathSet())
 	{
diff --git a/Engine/Source/Editor/UnrealEd/Private/PlayLevelNewProcess.cpp b/Engine/Source/Editor/UnrealEd/Private/PlayLevelNewProcess.cpp
--- a/Engine/Source/Editor/UnrealEd/Private/PlayLevelNewProcess.cpp
+++ b/Engine/Source/Editor/UnrealEd/Private/PlayLevelNewProcess.cpp
@@ -354,7 +354,7 @@ void UEditorEngine::LaunchNewProcess(const FRequestPlaySessionParams& InParams,
 	{
 		CommandLineArgFlags = ECommandLineArgumentFlags::ClientContext;
 	}
-	FCommandLine::BuildSubprocessCommandLine(CommandLineArgFlags, SubprocessCommandLine);
+	FCommandLine::BuildSubprocessCommandLine(CommandLineArgFlags, false /*bOnlyInherited*/, SubprocessCommandLine);
 	NamedArguments.Add(TEXT("SubprocessCommandLine"), *SubprocessCommandLine);
 	NamedArguments.Add(TEXT("CommandLineParams"), CommandLine);
 
diff --git a/Engine/Source/Runtime/Core/Private/Misc/CommandLine.cpp b/Engine/Source/Runtime/Core/Private/Misc/CommandLine.cpp
--- a/Engine/Source/Runtime/Core/Private/Misc/CommandLine.cpp
+++ b/Engine/Source/Runtime/Core/Private/Misc/CommandLine.cpp
@@ -395,22 +395,25 @@ const FString& FCommandLine::GetSubprocessCommandline()
 	return GetSubprocessCommandLine_Internal(ECommandLineArgumentFlags::AllContexts);
 }
 
-void FCommandLine::BuildSubprocessCommandLine(ECommandLineArgumentFlags ApplicationContextFlags, FStringBuilderBase& OutCommandline)
+void FCommandLine::BuildSubprocessCommandLine(ECommandLineArgumentFlags ApplicationContextFlags, bool bOnlyInherited, FStringBuilderBase& OutCommandline)
 {
 	if ((ApplicationContextFlags & ECommandLineArgumentFlags::AllContexts) == ECommandLineArgumentFlags::None)
 	{
 		return;
 	}
 
-	// Append all explicitly supplied subprocess arguments
-	for (int ContextPermutation = 1; ContextPermutation < (int)ECommandLineArgumentFlags::AllContexts; ++ContextPermutation)
+	if (!bOnlyInherited)
 	{
-		if (ContextPermutation & (int)ApplicationContextFlags)
+		// Append all explicitly supplied subprocess arguments
+		for (int ContextPermutation = 1; ContextPermutation < (int)ECommandLineArgumentFlags::AllContexts; ++ContextPermutation)
 		{
-			OutCommandline.Append(GetSubprocessCommandLine_Internal((ECommandLineArgumentFlags)ContextPermutation));
+			if (ContextPermutation & (int)ApplicationContextFlags)
+			{
+				OutCommandline.Append(GetSubprocessCommandLine_Internal((ECommandLineArgumentFlags)ContextPermutation));
+			}
 		}
+		OutCommandline.Append(GetSubprocessCommandLine_Internal(ECommandLineArgumentFlags::AllContexts));
 	}
-	OutCommandline.Append(GetSubprocessCommandLine_Internal(ECommandLineArgumentFlags::AllContexts));
 
 	// Append all inherited arguments
 	const TCHAR* Stream = Get();
diff --git a/Engine/Source/Runtime/Core/Public/Misc/CommandLine.h b/Engine/Source/Runtime/Core/Public/Misc/CommandLine.h
--- a/Engine/Source/Runtime/Core/Public/Misc/CommandLine.h
+++ b/Engine/Source/Runtime/Core/Public/Misc/CommandLine.h
@@ -129,9 +129,10 @@ struct FCommandLine
 	* Builds a commandline of inheritable and subprocess arguments for a specified application context.
 	*
 	* @param Flags 				The flags indicating the subprocess context we want a commandline for.
-	* @param Flags 				The flags indicating the subprocess context we want a commandline for.
+	* @param bOnlyInherited 	If true, only add inherited arguments, not explicitly supplied subprocess arguments.
+	* @param OutCommandline 	[out] The composed commandline.
 	*/
-	CORE_API static void BuildSubprocessCommandLine(ECommandLineArgumentFlags ApplicationContextFlags, FStringBuilderBase& OutCommandline);
+	CORE_API static void BuildSubprocessCommandLine(ECommandLineArgumentFlags ApplicationContextFlags, bool bOnlyInherited, FStringBuilderBase& OutCommandline);
 
 	/** 
 	* Removes the executable name from the passed CmdLine, denoted by parentheses.
diff --git a/Engine/Source/Runtime/Engine/Private/ShaderCompiler/ShaderCompiler.cpp b/Engine/Source/Runtime/Engine/Private/ShaderCompiler/ShaderCompiler.cpp
--- a/Engine/Source/Runtime/Engine/Private/ShaderCompiler/ShaderCompiler.cpp
+++ b/Engine/Source/Runtime/Engine/Private/ShaderCompiler/ShaderCompiler.cpp
@@ -1593,7 +1593,7 @@ FProcHandle FShaderCompilingManager::LaunchWorker(const FString& WorkingDirector
 		}
 	}
 	TStringBuilder<64> SubprocessCommandLine;
-	FCommandLine::BuildSubprocessCommandLine(ECommandLineArgumentFlags::ProgramContext, SubprocessCommandLine);
+	FCommandLine::BuildSubprocessCommandLine(ECommandLineArgumentFlags::ProgramContext, false /*bOnlyInherited*/, SubprocessCommandLine);
 	WorkerParameters += SubprocessCommandLine;
 
 #if USE_SHADER_COMPILER_WORKER_TRACE
diff --git a/Engine/Source/Runtime/Engine/Private/ShaderCompiler/ShaderCompilerDistributed.cpp b/Engine/Source/Runtime/Engine/Private/ShaderCompiler/ShaderCompilerDistributed.cpp
--- a/Engine/Source/Runtime/Engine/Private/ShaderCompiler/ShaderCompilerDistributed.cpp
+++ b/Engine/Source/Runtime/Engine/Private/ShaderCompiler/ShaderCompilerDistributed.cpp
@@ -199,7 +199,7 @@ void FShaderCompileDistributedThreadRunnable_Interface::DispatchShaderCompileJob
 	TaskCommandData.InputFileName = InputFilePath;
 	TaskCommandData.OutputFileName = OutputFilePath;
 	TStringBuilder<64> SubprocessCommandLine;
-	FCommandLine::BuildSubprocessCommandLine(ECommandLineArgumentFlags::ProgramContext, SubprocessCommandLine);
+	FCommandLine::BuildSubprocessCommandLine(ECommandLineArgumentFlags::ProgramContext, false /*bOnlyInherited*/, SubprocessCommandLine);
 	TaskCommandData.ExtraCommandArgs = FString::Printf(TEXT("%s%s"), *SubprocessCommandLine, GIsBuildMachine ? TEXT(" -buildmachine") : TEXT(""));
 	TaskCommandData.Dependencies = GetDependencyFilesForJobs(JobsToSerialize);
 
