Allow launcher profiles to run when not cooking, if we are going to import a zen snapshot

#rnx

diff --git a/Engine/Source/Developer/LauncherServices/Private/Profiles/LauncherProfile.h b/Engine/Source/Developer/LauncherServices/Private/Profiles/LauncherProfile.h
--- a/Engine/Source/Developer/LauncherServices/Private/Profiles/LauncherProfile.h
+++ b/Engine/Source/Developer/LauncherServices/Private/Profiles/LauncherProfile.h
@@ -2938,8 +2938,8 @@ protected:
 			ValidationErrors.Add(ELauncherProfileValidationErrors::DeployedDeviceGroupRequired);
 		}
 
-		// Deploy: deployment by copying to devices requires cooking by the book
-		if ((DeploymentMode == ELauncherProfileDeploymentModes::CopyToDevice) && ((CookMode != ELauncherProfileCookModes::ByTheBook)&&(CookMode!=ELauncherProfileCookModes::ByTheBookInEditor)) && !IsUsingZenPakStreaming())
+		// Deploy: deployment by copying to devices requires cooking by the book unless the data is available from zen
+		if ((DeploymentMode == ELauncherProfileDeploymentModes::CopyToDevice) && ((CookMode != ELauncherProfileCookModes::ByTheBook)&&(CookMode!=ELauncherProfileCookModes::ByTheBookInEditor)) && !IsUsingZenPakStreaming() && !IsImportingZenSnapshot())
 		{
 			ValidationErrors.Add(ELauncherProfileValidationErrors::CopyToDeviceRequiresCookByTheBook);
 		}
