Add support for running automated tests through Project Launcher.
- automated test launcher extension base class.
- example BootTest launcher extension.

diff --git a/Engine/Plugins/Developer/ProjectLauncher/Source/CommonLaunchExtensions/Private/BootTest/BootTestLaunchExtension.cpp b/Engine/Plugins/Developer/ProjectLauncher/Source/CommonLaunchExtensions/Private/BootTest/BootTestLaunchExtension.cpp
new file mode 100644
--- /dev/null
+++ b/Engine/Plugins/Developer/ProjectLauncher/Source/CommonLaunchExtensions/Private/BootTest/BootTestLaunchExtension.cpp
@@ -0,0 +1,69 @@
+// Copyright Epic Games, Inc. All Rights Reserved.
+
+#include "BootTest/BootTestLaunchExtension.h"
+
+
+#define LOCTEXT_NAMESPACE "FBootTestLaunchExtensionInstance"
+
+const TCHAR* FBootTestLaunchExtensionInstance::BootTestInternalName = TEXT("BootTestExtension.BootTest");
+
+
+void FBootTestLaunchExtensionInstance::CustomizeTree( ProjectLauncher::FLaunchProfileTreeData& ProfileTreeData )
+{
+	// note: Windowed option is here mostly as an example of adding UI items for a particular profile - your tests may not need it
+
+	auto IsVisible = [this]()
+	{
+		return IsTestActive();
+	};
+
+	ProfileTreeData.AddHeading(GetExtension()->GetDisplayName())
+		.AddBoolean(LOCTEXT("WindowedLabel","Windowed"),
+		{
+			.GetValue = [this]()          { return GetConfigBool(EConfig::PerProfile, TEXT("Windowed"));},
+			.SetValue = [this](bool bVal) { SetConfigBool(EConfig::PerProfile, TEXT("Windowed"), bVal ); },
+			.IsVisible = IsVisible,
+		}
+	);
+}
+
+void FBootTestLaunchExtensionInstance::CustomizeAutomatedTestCommandLine( FString& InOutCommandLine )
+{
+	bool bWindowed = GetConfigBool(EConfig::PerProfile, TEXT("Windowed"));
+	if (!bWindowed)
+	{
+		InOutCommandLine += TEXT(" -windowmode=Fullscreen");
+	}
+}
+
+const FString FBootTestLaunchExtensionInstance::GetTestInternalName() const
+{
+	return BootTestInternalName;
+}
+
+void FBootTestLaunchExtensionInstance::OnTestAdded( ILauncherProfileAutomatedTestRef AutomatedTest )
+{
+	AutomatedTest->SetTests(TEXT("UE.BootTest"));
+	AutomatedTest->SetPriority(1000); // Boot Test is likely the most lightweight test there is, so we want to run it first - give it a high priority
+}
+
+
+
+
+TSharedPtr<ProjectLauncher::FLaunchExtensionInstance> FBootTestLaunchExtension::CreateInstanceForProfile( ProjectLauncher::FLaunchExtensionInstance::FArgs& InArgs )
+{
+	return MakeShared<FBootTestLaunchExtensionInstance>(InArgs);
+}
+
+const TCHAR* FBootTestLaunchExtension::GetInternalName() const
+{
+	return TEXT("BootTest");
+}
+
+FText FBootTestLaunchExtension::GetDisplayName() const
+{
+	return LOCTEXT("ExtensionName", "Automated Boot Test");
+}
+
+
+#undef LOCTEXT_NAMESPACE
diff --git a/Engine/Plugins/Developer/ProjectLauncher/Source/CommonLaunchExtensions/Private/BootTest/BootTestLaunchExtension.h b/Engine/Plugins/Developer/ProjectLauncher/Source/CommonLaunchExtensions/Private/BootTest/BootTestLaunchExtension.h
new file mode 100644
--- /dev/null
+++ b/Engine/Plugins/Developer/ProjectLauncher/Source/CommonLaunchExtensions/Private/BootTest/BootTestLaunchExtension.h
@@ -0,0 +1,33 @@
+// Copyright Epic Games, Inc. All Rights Reserved.
+
+#pragma once
+
+#include "Extension/LaunchExtension.h"
+
+
+
+class FBootTestLaunchExtensionInstance : public ProjectLauncher::FAutomatedTestLaunchExtensionInstance
+{
+public:
+	FBootTestLaunchExtensionInstance( FArgs& InArgs ) : ProjectLauncher::FAutomatedTestLaunchExtensionInstance(InArgs) {};
+	virtual ~FBootTestLaunchExtensionInstance() = default;
+
+	virtual void OnTestAdded( ILauncherProfileAutomatedTestRef AutomatedTest ) override;
+	virtual const FString GetTestInternalName() const override;
+
+	virtual void CustomizeTree( ProjectLauncher::FLaunchProfileTreeData& ProfileTreeData ) override;
+	virtual void CustomizeAutomatedTestCommandLine( FString& InOutCommandLine ) override;
+
+private:
+	static const TCHAR* BootTestInternalName;
+
+};
+
+
+class FBootTestLaunchExtension : public ProjectLauncher::FLaunchExtension
+{
+public:
+	virtual TSharedPtr<ProjectLauncher::FLaunchExtensionInstance> CreateInstanceForProfile( ProjectLauncher::FLaunchExtensionInstance::FArgs& InArgs ) override;
+	virtual const TCHAR* GetInternalName() const override;
+	virtual FText GetDisplayName() const override;
+};
\ No newline at end of file
diff --git a/Engine/Plugins/Developer/ProjectLauncher/Source/CommonLaunchExtensions/Private/CommonLaunchExtensionsModule.cpp b/Engine/Plugins/Developer/ProjectLauncher/Source/CommonLaunchExtensions/Private/CommonLaunchExtensionsModule.cpp
--- a/Engine/Plugins/Developer/ProjectLauncher/Source/CommonLaunchExtensions/Private/CommonLaunchExtensionsModule.cpp
+++ b/Engine/Plugins/Developer/ProjectLauncher/Source/CommonLaunchExtensions/Private/CommonLaunchExtensionsModule.cpp
@@ -5,6 +5,7 @@
 #include "ProjectLauncherModule.h"
 #include "Globals/GlobalsLaunchExtension.h"
 #include "Insights/InsightsLaunchExtension.h"
+#include "BootTest/BootTestLaunchExtension.h"
 
 class FCommonLaunchExtensionsModule : public IModuleInterface
 {
@@ -14,9 +15,11 @@ public:
 	{
 		Globals = MakeShared<FGlobalsLaunchExtension>();
 		Insights = MakeShared<FInsightsLaunchExtension>();
+		BootTest = MakeShared<FBootTestLaunchExtension>();
 
 		IProjectLauncherModule::Get().RegisterExtension(Globals.ToSharedRef());
 		IProjectLauncherModule::Get().RegisterExtension(Insights.ToSharedRef());
+		IProjectLauncherModule::Get().RegisterExtension(BootTest.ToSharedRef());
 	}
 
 	virtual void ShutdownModule() override
@@ -25,15 +28,18 @@ public:
 		{
 			ProjectLauncher->UnregisterExtension(Globals.ToSharedRef());
 			ProjectLauncher->UnregisterExtension(Insights.ToSharedRef());
+			ProjectLauncher->UnregisterExtension(BootTest.ToSharedRef());
 		}
 
 		Globals.Reset();
 		Insights.Reset();
+		BootTest.Reset();
 	}
 
 private:
 	TSharedPtr<FGlobalsLaunchExtension> Globals;
 	TSharedPtr<FInsightsLaunchExtension> Insights;
+	TSharedPtr<FBootTestLaunchExtension> BootTest;
 
 };
 
diff --git a/Engine/Plugins/Developer/ProjectLauncher/Source/ProjectLauncher/Private/Extension/LaunchExtension.cpp b/Engine/Plugins/Developer/ProjectLauncher/Source/ProjectLauncher/Private/Extension/LaunchExtension.cpp
--- a/Engine/Plugins/Developer/ProjectLauncher/Source/ProjectLauncher/Private/Extension/LaunchExtension.cpp
+++ b/Engine/Plugins/Developer/ProjectLauncher/Source/ProjectLauncher/Private/Extension/LaunchExtension.cpp
@@ -75,6 +75,47 @@ namespace ProjectLauncher
 		}
 	}
 
+	void ApplyExtensionVariablesForAutomatedTest( const ILauncherProfileAutomatedTestRef& InAutomatedTest, const ILauncherProfileRef& InProfile, FString& InOutCommandLine, TSharedRef<FModel> InModel )
+	{
+		// take a snapshot of the extension instances for this profile
+		TArray<TSharedPtr<FLaunchExtensionInstance>> ExtensionInstances;
+		{
+			FScopeLock Lock(&GExtensionsCS);
+			if (GExtensionInstances.Contains(InProfile))
+			{
+				ExtensionInstances = GExtensionInstances[InProfile];
+			}
+		}
+
+
+		for (TSharedPtr<FLaunchExtensionInstance> ExtensionInstance : ExtensionInstances)
+		{
+			if (ExtensionInstance.IsValid())
+			{
+				// apply variable substitutions
+				TArray<FString> Variables;
+				if (ExtensionInstance->GetExtensionVariables(Variables))
+				{
+					for (const FString& Variable : Variables)
+					{
+						if (InOutCommandLine.Contains(Variable, ESearchCase::IgnoreCase))
+						{
+							FString VariableValue;
+							if (ExtensionInstance->GetExtensionVariableValue( Variable, VariableValue ))
+							{
+								InOutCommandLine.ReplaceInline( *Variable, *VariableValue, ESearchCase::IgnoreCase );
+							}
+						}
+					}
+				}
+
+				// allow for advanced command line customization
+				ExtensionInstance->CustomizeAutomatedTestCommandLine( InAutomatedTest, InOutCommandLine );
+			}
+		}
+	}
+
+
 
 	void RemoveExtensionInstancesForProfile( const ILauncherProfileRef& Profile )
 	{
@@ -511,6 +552,83 @@ namespace ProjectLauncher
 	{
 		PropertyChangedDelegate.Broadcast();
 	}
+
+	FText FLaunchExtensionInstance::GetCustomExtensionMenuSection() const
+	{
+		return Extension->GetDisplayName();
+	}
+
+
+
+
+
+
+
+
+	FAutomatedTestLaunchExtensionInstance::FAutomatedTestLaunchExtensionInstance( FArgs& InArgs )
+		: FLaunchExtensionInstance(InArgs)
+	{
+	}
+
+	void FAutomatedTestLaunchExtensionInstance::MakeCustomExtensionSubmenu( FMenuBuilder& MenuBuilder )
+	{
+		auto ToggleTestEnabled = [this]()
+		{
+			ILauncherProfileAutomatedTestPtr ExistingAutomatedTest = GetTest();
+			if (ExistingAutomatedTest.IsValid())
+			{
+				OnTestRemoved(ExistingAutomatedTest.ToSharedRef());
+				GetProfile()->RemoveAutomatedTest(*GetTestInternalName());
+			}
+			else
+			{
+				ILauncherProfileAutomatedTestRef AutomatedTest = GetProfile()->FindOrAddAutomatedTest(*GetTestInternalName());
+				OnTestAdded(AutomatedTest);
+			}
+
+			BroadcastPropertyChanged();
+		};
+
+		MenuBuilder.AddMenuEntry(
+			GetExtension()->GetDisplayName(),
+			FText::GetEmpty(),
+			FSlateIcon(),
+			FUIAction(
+				FExecuteAction::CreateLambda(ToggleTestEnabled),
+				FCanExecuteAction(),
+				FIsActionChecked::CreateRaw(this, &FAutomatedTestLaunchExtensionInstance::IsTestActive )
+			),
+			NAME_None,
+			EUserInterfaceActionType::Check);
+	}
+
+	bool FAutomatedTestLaunchExtensionInstance::IsTestActive() const
+	{
+		return GetTest() != nullptr;
+	}
+
+	ILauncherProfileAutomatedTestPtr FAutomatedTestLaunchExtensionInstance::GetTest() const
+	{
+		return GetProfile()->GetAutomatedTest(*GetTestInternalName());
+	}
+
+	FText FAutomatedTestLaunchExtensionInstance::GetDefaultExtensionMenuSection()
+	{
+		return LOCTEXT("AutotestMenuSection", "Automated Tests");
+	}
+
+	FText FAutomatedTestLaunchExtensionInstance::GetCustomExtensionMenuSection() const
+	{
+		return GetDefaultExtensionMenuSection();
+	}
+
+	void FAutomatedTestLaunchExtensionInstance::CustomizeAutomatedTestCommandLine( const ILauncherProfileAutomatedTestRef& InAutomatedTest, FString& InOutCommandLine )
+	{
+		if (InAutomatedTest == GetTest())
+		{
+			CustomizeAutomatedTestCommandLine(InOutCommandLine);
+		}
+	}
 }
 
 #undef LOCTEXT_NAMESPACE
diff --git a/Engine/Plugins/Developer/ProjectLauncher/Source/ProjectLauncher/Private/Model/ProjectLauncherModel.cpp b/Engine/Plugins/Developer/ProjectLauncher/Source/ProjectLauncher/Private/Model/ProjectLauncherModel.cpp
--- a/Engine/Plugins/Developer/ProjectLauncher/Source/ProjectLauncher/Private/Model/ProjectLauncherModel.cpp
+++ b/Engine/Plugins/Developer/ProjectLauncher/Source/ProjectLauncher/Private/Model/ProjectLauncherModel.cpp
@@ -39,6 +39,7 @@ namespace ProjectLauncher
 	{
 		// register callbacks
 		ProfileManager->OnPostProcessLaunchCommandLine().AddRaw( this, &FModel::OnModifyLaunchCommandLine );
+		ProfileManager->OnPostProcessAutomatedTestCommandLine().AddRaw( this, &FModel::OnModifyAutomatedTestCommandLine );
 
 		ProfileManager->OnProfileAdded().AddRaw(this, &FModel::HandleProfileManagerProfileAdded);
 		ProfileManager->OnProfileRemoved().AddRaw(this, &FModel::HandleProfileManagerProfileRemoved);
@@ -83,7 +84,8 @@ namespace ProjectLauncher
 			delete Pair.Value.Config;
 		}
 
-		ProfileManager->OnPostProcessLaunchCommandLine().RemoveAll( this );
+		ProfileManager->OnPostProcessLaunchCommandLine().RemoveAll(this);
+		ProfileManager->OnPostProcessAutomatedTestCommandLine().RemoveAll(this);
 
 		ProfileManager->OnProfileAdded().RemoveAll(this);
 		ProfileManager->OnProfileRemoved().RemoveAll(this);
@@ -719,6 +721,15 @@ namespace ProjectLauncher
 		ApplyExtensionVariables(InProfile, InOutCommandLine, AsShared() );
 	}
 
+
+	extern void ApplyExtensionVariablesForAutomatedTest( const ILauncherProfileAutomatedTestRef& InAutomatedTest, const ILauncherProfileRef& InProfile, FString& InOutCommandLine, TSharedRef<FModel> InModel );
+
+	void FModel::OnModifyAutomatedTestCommandLine( const ILauncherProfileAutomatedTestRef& InAutomatedTest, const ILauncherProfileRef& InProfile, FString& InOutCommandLine )
+	{
+		ApplyExtensionVariablesForAutomatedTest(InAutomatedTest, InProfile, InOutCommandLine, AsShared() );
+	}
+
+
 	bool FModel::AreExtensionsEnabled() const
 	{
 #if WITH_EDITOR
diff --git a/Engine/Plugins/Developer/ProjectLauncher/Source/ProjectLauncher/Private/ProfileTree/GenericProfileTreeBuilder.cpp b/Engine/Plugins/Developer/ProjectLauncher/Source/ProjectLauncher/Private/ProfileTree/GenericProfileTreeBuilder.cpp
--- a/Engine/Plugins/Developer/ProjectLauncher/Source/ProjectLauncher/Private/ProfileTree/GenericProfileTreeBuilder.cpp
+++ b/Engine/Plugins/Developer/ProjectLauncher/Source/ProjectLauncher/Private/ProfileTree/GenericProfileTreeBuilder.cpp
@@ -494,6 +494,7 @@ namespace ProjectLauncher
 				.GetValue = [this]()			{ return GetIsRunning(); },
 				.SetValue = [this](bool bValue)	{ SetIsRunning(bValue); },
 				.GetDefaultValue = [this]()		{ return GetIsRunning(DefaultProfile); },
+				.IsVisible = [this]()           { return Profile->GetAutomatedTests().Num() == 0; } // run is implied when there are automated tests
 			}
 		);
 	}
diff --git a/Engine/Plugins/Developer/ProjectLauncher/Source/ProjectLauncher/Private/Styling/ProjectLauncherStyle.cpp b/Engine/Plugins/Developer/ProjectLauncher/Source/ProjectLauncher/Private/Styling/ProjectLauncherStyle.cpp
--- a/Engine/Plugins/Developer/ProjectLauncher/Source/ProjectLauncher/Private/Styling/ProjectLauncherStyle.cpp
+++ b/Engine/Plugins/Developer/ProjectLauncher/Source/ProjectLauncher/Private/Styling/ProjectLauncherStyle.cpp
@@ -88,6 +88,8 @@ void FProjectLauncherStyle::Initialize()
 	StyleSet->Set("Icons.Task.Deploy", new CORE_IMAGE_BRUSH_SVG( "../Editor/Slate/Starship/Common/DeviceManager", Icon16x16 ) );
 	StyleSet->Set("Icons.Task.Archive", new CORE_IMAGE_BRUSH_SVG("../Editor/Slate/Starship/Common/SaveCurrent", Icon16x16));
 	StyleSet->Set("Icons.Task.Package", new CORE_IMAGE_BRUSH_SVG("../Editor/Slate/Starship/Common/ProjectPackage", Icon16x16));
+	StyleSet->Set("Icons.Task.TestAutomation", new CORE_IMAGE_BRUSH_SVG("Starship/Common/Test", Icon16x16));
+
 
 	StyleSet->Set("Profile.NoPlatform", new CORE_IMAGE_BRUSH_SVG("Starship/Launcher/PaperAirplane", Icon24x24));
 	StyleSet->Set("Profile.NoPlatform.Large", new CORE_IMAGE_BRUSH_SVG("Starship/Launcher/PaperAirplane", Icon64x64));
@@ -122,7 +124,6 @@ const FName& FProjectLauncherStyle::GetStyleSetName()
 
 const FSlateBrush* FProjectLauncherStyle::GetBrushForTask(ILauncherTaskPtr Task)
 {
-	// fixme: use bespoke icons
 	if (Task.IsValid())
 	{
 		FString TaskName = Task->GetName();
@@ -162,6 +163,10 @@ const FSlateBrush* FProjectLauncherStyle::GetBrushForTask(ILauncherTaskPtr Task)
 		{
 			return GetBrush("Icons.Task.Run");
 		}
+		else if (TaskName.Contains(TEXT("Automated Testing")))
+		{
+			return GetBrush("Icons.Task.TestAutomation");
+		}
 	}
 
 	return GetBrush("Icons.Task.Package");
diff --git a/Engine/Plugins/Developer/ProjectLauncher/Source/ProjectLauncher/Private/Widgets/Profiles/SCustomLaunchCustomProfileEditor.cpp b/Engine/Plugins/Developer/ProjectLauncher/Source/ProjectLauncher/Private/Widgets/Profiles/SCustomLaunchCustomProfileEditor.cpp
--- a/Engine/Plugins/Developer/ProjectLauncher/Source/ProjectLauncher/Private/Widgets/Profiles/SCustomLaunchCustomProfileEditor.cpp
+++ b/Engine/Plugins/Developer/ProjectLauncher/Source/ProjectLauncher/Private/Widgets/Profiles/SCustomLaunchCustomProfileEditor.cpp
@@ -393,14 +393,14 @@ TSharedRef<SWidget> SCustomLaunchCustomProfileEditor::MakeExtensionsMenu()
 
 #if 1
 		// flat list of extensions - easier to see but risks getting cluttered later on
-		MenuBuilder.BeginSection(NAME_None, ExtensionInstance->GetExtension()->GetDisplayName());
+		MenuBuilder.BeginSection(NAME_None, ExtensionInstance->GetCustomExtensionMenuSection());
 		ExtensionInstance->MakeCustomExtensionSubmenu(MenuBuilder);
 		MenuBuilder.EndSection();
 #else 
 		// extensions in submenus - more scalable, but harder to see everything at a glance
 		const bool bInOpenSubMenuOnClick = false;
 		MenuBuilder.AddSubMenu(
-			ExtensionInstance->GetExtension()->GetDisplayName()),
+			ExtensionInstance->GetCustomExtensionMenuSection()),
 			FText::GetEmpty(),
 			FNewMenuDelegate::CreateSP( ExtensionInstance.ToSharedRef(), &FLaunchExtensionInstance::MakeExtensionSubmenu),
 			bInOpenSubMenuOnClick,
diff --git a/Engine/Plugins/Developer/ProjectLauncher/Source/ProjectLauncher/Private/Widgets/Profiles/SCustomLaunchProfilesPanel.cpp b/Engine/Plugins/Developer/ProjectLauncher/Source/ProjectLauncher/Private/Widgets/Profiles/SCustomLaunchProfilesPanel.cpp
--- a/Engine/Plugins/Developer/ProjectLauncher/Source/ProjectLauncher/Private/Widgets/Profiles/SCustomLaunchProfilesPanel.cpp
+++ b/Engine/Plugins/Developer/ProjectLauncher/Source/ProjectLauncher/Private/Widgets/Profiles/SCustomLaunchProfilesPanel.cpp
@@ -677,6 +677,10 @@ const FSlateBrush* SCustomLaunchProfilesPanel::GetLaunchButtonImage() const
 	{
 		return FAppStyle::Get().GetBrush("Icons.Error");
 	}
+	else if (Profile->GetAutomatedTests().Num() > 0)
+	{
+		return FProjectLauncherStyle::Get().GetBrush("Icons.Task.TestAutomation");
+	}
 	else if (Profile->GetLaunchMode() != ELauncherProfileLaunchModes::DoNotLaunch)
 	{
 		return FProjectLauncherStyle::Get().GetBrush("Icons.Task.Launch");
diff --git a/Engine/Plugins/Developer/ProjectLauncher/Source/ProjectLauncher/Public/Extension/LaunchExtension.h b/Engine/Plugins/Developer/ProjectLauncher/Source/ProjectLauncher/Public/Extension/LaunchExtension.h
--- a/Engine/Plugins/Developer/ProjectLauncher/Source/ProjectLauncher/Public/Extension/LaunchExtension.h
+++ b/Engine/Plugins/Developer/ProjectLauncher/Source/ProjectLauncher/Public/Extension/LaunchExtension.h
@@ -76,6 +76,11 @@ namespace ProjectLauncher
 		 */
 		virtual void CustomizeLaunchCommandLine( FString& InOutCommandLine ) {};
 
+		/**
+		 * Advanced hook to allow any advanced modification of the command line when an automated test is launched
+		 */
+		virtual void CustomizeAutomatedTestCommandLine( const ILauncherProfileAutomatedTestRef& InAutomatedTest, FString& InOutCommandLine ) {};
+
 		/**
 		 * Notification callback when a property has been changed in our profile
 		 */
@@ -86,6 +91,11 @@ namespace ProjectLauncher
 		 */
 		virtual bool HasCustomExtensionMenu() const { return false; }
 
+		/**
+		 * Get the custom extension menu section. Default is the extension's friendly name
+		 */
+		UE_API virtual FText GetCustomExtensionMenuSection() const;
+
 		/*
 		 * Populate the custom extension menu for this extension. This menu would typically contain
 		 * items that might enable bespoke tree customization options for example.
@@ -251,7 +261,7 @@ namespace ProjectLauncher
 		/**
 		 * Signal to any listeners that a property in the profile has changed
 		 */
-		void BroadcastPropertyChanged() const;
+		UE_API void BroadcastPropertyChanged() const;
 
 	private:
 		UE_API bool IsParameterGroup( const FString& InParameter ) const;
@@ -269,6 +279,71 @@ namespace ProjectLauncher
 	};
 
 
+	/**
+	 * Helper base class for an automated test launch extension
+	 * Creates an item in the extension menu that will toggle the test on and off
+	 * note: It is not necessary for an automated test extensions to use this class
+	 */
+	class FAutomatedTestLaunchExtensionInstance : public FLaunchExtensionInstance
+	{
+	public:
+		UE_API FAutomatedTestLaunchExtensionInstance( FArgs& InArgs );
+
+		/*
+		 * Populate the custom extension menu for this automated test extension. This will add a menu item to toggle the test on and off
+		 */
+		UE_API virtual void MakeCustomExtensionSubmenu( FMenuBuilder& MenuBuilder ) override;
+
+		/**
+		 * Get the custom extension menu section. Default is a common section for all automated tests
+		 */
+		UE_API FText GetCustomExtensionMenuSection() const override;
+
+		/**
+		 * Returns the common menu section name for all automated tests
+		 * Can be used to insert things into this menu that may not inherit from this class
+		 */
+		UE_API static FText GetDefaultExtensionMenuSection();
+		
+		/**
+		 * Determine if this automated test is currently active
+		 */
+		UE_API bool IsTestActive() const;
+
+	protected:
+		UE_API void CustomizeAutomatedTestCommandLine( const ILauncherProfileAutomatedTestRef& InAutomatedTest, FString& InOutCommandLine ) override;
+		virtual bool HasCustomExtensionMenu() const { return true; }
+
+		/**
+		 * Advanced hook to allow any advanced modification of the command line when this extension's automated test is launched
+		 */
+		virtual void CustomizeAutomatedTestCommandLine( FString& InOutCommandLine ) {}
+
+		/**
+		 * Hook to allow any modification of the automated test & profile once it has been created. At least the test name should be specified
+		 */
+		virtual void OnTestAdded( ILauncherProfileAutomatedTestRef AutomatedTest ) = 0;
+
+		/**
+		 * Advanced hook to allow any modification of the profile just before the automated test is removed
+		 */
+		virtual void OnTestRemoved( ILauncherProfileAutomatedTestRef AutomatedTest ) {}
+
+		/**
+		 * Access this extension's automated test if it has been created
+		 */
+		UE_API ILauncherProfileAutomatedTestPtr GetTest() const;
+
+		/**
+		 * Internal name for the automated test. Must be globally unique, so recommended to include the extension name
+		 */
+		virtual const FString GetTestInternalName() const = 0;
+	};
+
+
+
+
+
 
 	/**
 	 * Base class for a launch extension.
diff --git a/Engine/Plugins/Developer/ProjectLauncher/Source/ProjectLauncher/Public/Model/ProjectLauncherModel.h b/Engine/Plugins/Developer/ProjectLauncher/Source/ProjectLauncher/Public/Model/ProjectLauncherModel.h
--- a/Engine/Plugins/Developer/ProjectLauncher/Source/ProjectLauncher/Public/Model/ProjectLauncherModel.h
+++ b/Engine/Plugins/Developer/ProjectLauncher/Source/ProjectLauncher/Public/Model/ProjectLauncherModel.h
@@ -199,6 +199,8 @@ namespace ProjectLauncher
 		UE_API void ReadProjectSettingsFromConfig( FConfigCacheIni& InConfig, const FString& InProjectPath, FProjectSettings& OutResult );
 
 		UE_API void OnModifyLaunchCommandLine( const ILauncherProfileRef& InProfile, FString& InOutCommandLine );
+		UE_API void OnModifyAutomatedTestCommandLine( const ILauncherProfileAutomatedTestRef& InAutomatedTest, const ILauncherProfileRef& InProfile, FString& InOutCommandLine );
+
 
 		UE_API TArray<FString> GetAndCacheMapPaths(const FString& InOptionalProjectPath);
 
diff --git a/Engine/Source/Developer/LauncherServices/Private/Launcher/LauncherUATTask.h b/Engine/Source/Developer/LauncherServices/Private/Launcher/LauncherUATTask.h
--- a/Engine/Source/Developer/LauncherServices/Private/Launcher/LauncherUATTask.h
+++ b/Engine/Source/Developer/LauncherServices/Private/Launcher/LauncherUATTask.h
@@ -17,7 +17,7 @@ class FLauncherUATTask
 {
 public:
 
-	FLauncherUATTask( const FString& InCommandLine, const FString& InName, const FString& InDesc, void* InReadPipe, void* InWritePipe, const FString& InEditorExe, FProcHandle& InProcessHandle, ILauncherWorker* InWorker, const FString& InCommandEnd, const FString& InBuildCookRunPreCommand)
+	FLauncherUATTask( const FString& InCommandLine, const FString& InName, const FString& InDesc, void* InReadPipe, void* InWritePipe, const FString& InEditorExe, FProcHandle& InProcessHandle, ILauncherWorker* InWorker, const FString& InCommandEnd, const FString& InBuildCookRunPreCommand, const FString& InBuildCookRunPostCommand)
 		: FLauncherTask(InName, InDesc)
 
 		, CommandLine(InCommandLine)
@@ -27,6 +27,7 @@ public:
 		, ProcessHandle(InProcessHandle)
 		, CommandText(InCommandEnd)
 		, PreCommand(InBuildCookRunPreCommand)
+		, PostCommand(InBuildCookRunPostCommand)
 	{
 		InWorker->OnOutputReceived().AddRaw(this, &FLauncherUATTask::HandleOutputReceived);
 	}
@@ -82,6 +83,14 @@ protected:
 		// specialized command arguments for this particular task
 		UATCommandLine += CommandLine;
 
+		// add any post command line
+		if (!PostCommand.IsEmpty())
+		{
+			UATCommandLine += TEXT(" ");
+			UATCommandLine += PostCommand;
+		}
+
+
 		// launch UAT and monitor its progress
 		ProcessHandle = FPlatformProcess::CreateProc(*(ExecutablePath / Executable), *UATCommandLine, false, true, true, NULL, 0, *ExecutablePath, WritePipe, ReadPipe);
 
@@ -133,6 +142,9 @@ private:
 	// a UAT command to run before BuildCookRun (in the same instance as the BuildCookRun)
 	FString PreCommand;
 
+	// a UAT command to run after BuildCookRun (in the same insance as the BuildCookRun)
+	FString PostCommand;
+
 	bool EndTextFound = false;
 };
 
diff --git a/Engine/Source/Developer/LauncherServices/Private/Launcher/LauncherWorker.cpp b/Engine/Source/Developer/LauncherServices/Private/Launcher/LauncherWorker.cpp
--- a/Engine/Source/Developer/LauncherServices/Private/Launcher/LauncherWorker.cpp
+++ b/Engine/Source/Developer/LauncherServices/Private/Launcher/LauncherWorker.cpp
@@ -1062,7 +1062,7 @@ FString FLauncherWorker::CreateUATCommand( const ILauncherProfileRef& InProfile,
 		}
 
 		// run
-		if (InProfile->GetLaunchMode() != ELauncherProfileLaunchModes::DoNotLaunch)
+		if (InProfile->GetLaunchMode() != ELauncherProfileLaunchModes::DoNotLaunch && Profile->GetAutomatedTests().Num() == 0)
 		{
 			UATCommand += TEXT(" -run ");
 
@@ -1122,17 +1122,130 @@ FString FLauncherWorker::CreateUATCommand( const ILauncherProfileRef& InProfile,
 		}
 	}
 
-	// wait for completion of UAT
-	FCommandDesc Desc;
-	FText Command = LOCTEXT("LauncherCompletionDesc", "UAT post launch cleanup");
-	Desc.Name = "Post Launch Task";
-	Desc.Desc = Command.ToString();
-	Desc.EndText = TEXT("********** LAUNCH COMPLETED **********");
-	OutCommands.Add(Desc);
+	return UATCommand;
+}
+
+
+FString FLauncherWorker::CreateUATCommandForAutomatedTests( const ILauncherProfileRef& InProfile, const TArray<FString>& InPlatforms, TArray<FCommandDesc>& OutCommands, FString& OutCommandStart )
+{
+	FString UATCommand;
+
+	if (Profile->GetAutomatedTests().Num() > 0)
+	{
+		TArray<ILauncherProfileAutomatedTestRef> PrioritizedAutomatedTests = Profile->GetAutomatedTests();
+		PrioritizedAutomatedTests.Sort([]( const ILauncherProfileAutomatedTestRef& A, const ILauncherProfileAutomatedTestRef& B)
+		{
+			return A->GetPriority() < B->GetPriority();
+		});
+
+		for (const ILauncherProfileAutomatedTestRef& AutomatedTest : PrioritizedAutomatedTests)
+		{
+			FString AutomatedTestCommand = CreateUATCommandForAutomatedTest(AutomatedTest, InProfile, InPlatforms, OutCommands, OutCommandStart);
+			UATCommand += AutomatedTestCommand + TEXT(" ");
+		}
+
+		FCommandDesc Desc;
+		FText Command = LOCTEXT("LauncherAutoTestDesc", "Automated Testing");
+		Desc.Name = "Automated Testing Task";
+		Desc.Desc = Command.ToString();
+		Desc.EndText = TEXT("********** AUTOMATED TESTS COMPLETED **********"); // @note: this isn't actually logged anywhere - this is the last real task so it will advance automatically
+		OutCommands.Add(Desc);
+	}
 
 	return UATCommand;
 }
 
+
+FString FLauncherWorker::CreateUATCommandForAutomatedTest( const ILauncherProfileAutomatedTestRef& InAutomatedTest, const ILauncherProfileRef& InProfile, const TArray<FString>& InPlatforms, TArray<FCommandDesc>& OutCommands, FString& OutCommandStart )
+{
+	// skip disabled tests
+	if (!InAutomatedTest->IsEnabled())
+	{
+		return FString();
+	}
+
+	ILauncherProfileManagerRef LauncherProfileManager = FModuleManager::LoadModuleChecked<ILauncherServicesModule>("LauncherServices").GetProfileManager();
+
+
+	// basic command parameters
+	FString AutomatedTestCommand = FString::Printf( TEXT("%s -utf8output -nop4 "), *InAutomatedTest->GetUATCommand() );
+	AutomatedTestCommand += FString::Printf(TEXT("-project=\"%s\" "), *InProfile->GetProjectPath());
+	AutomatedTestCommand += FString::Printf(TEXT("-configuration=%s "), LexToString(InProfile->GetBuildConfiguration()) );
+	AutomatedTestCommand += FString::Printf(TEXT("-tests=%s "), *InAutomatedTest->GetTests());
+
+	// collect platforms
+	FString ServerPlatforms = TEXT("");
+	FString Platforms = TEXT("");
+	for (int32 PlatformIndex = 0; PlatformIndex < InPlatforms.Num(); ++PlatformIndex)
+	{
+		const PlatformInfo::FTargetPlatformInfo* PlatformInfo = PlatformInfo::FindPlatformInfo(FName(*InPlatforms[PlatformIndex]));
+		if (ensure(PlatformInfo))
+		{
+			FString& PlatformString = (PlatformInfo->PlatformType == EBuildTargetType::Server) ? ServerPlatforms : Platforms;
+			PlatformString += TEXT("+");
+			PlatformString += PlatformInfo->DataDrivenPlatformInfo->UBTPlatformString;
+		}
+	}
+	if (ServerPlatforms.Len() > 0)
+	{
+		AutomatedTestCommand += FString::Printf(TEXT("-serverplatform=%s "), *ServerPlatforms.RightChop(1) );
+	}
+	if (Platforms.Len() > 0)
+	{
+		AutomatedTestCommand += FString::Printf(TEXT("-platform=%s "), *Platforms.RightChop(1) );
+	}
+
+	// specify the build to use
+	if (InProfile->GetAutomatedTestBuildPath().IsEmpty() || !InProfile->IsUsingAutomatedTestBuild())
+	{
+		AutomatedTestCommand += TEXT("-build=local ");
+	}
+	else
+	{
+		AutomatedTestCommand += FString::Printf( TEXT("-build=\"%s\" "), *InProfile->GetAutomatedTestBuildPath());
+	}
+
+	// use the local executables we just made if necessary
+	if (InProfile->ShouldBuild())
+	{
+		AutomatedTestCommand += TEXT("-dev ");
+	}
+
+	// specify the devices to use
+	ILauncherDeviceGroupPtr DeviceGroup = InProfile->GetDeployedDeviceGroup();
+	if (DeviceGroup.IsValid() && DeviceGroup->GetNumDevices() > 0)
+	{
+		FString DeviceNames;
+		for (const FString& DeviceId : DeviceGroup->GetDeviceIDs())
+		{
+			FString FinalDeviceId = DeviceId.TrimStartAndEnd().Replace(TEXT("@"), TEXT(":"));
+			if (!FinalDeviceId.IsEmpty())
+			{
+				DeviceNames += TEXT(",\"") + FinalDeviceId + TEXT("\"");
+			}
+		}
+		if (!DeviceNames.IsEmpty())
+		{
+			AutomatedTestCommand += FString::Printf(TEXT("-devices=%s "), *DeviceNames.RightChop(1));
+		}
+	}
+
+	// final extra command line parameters
+	if (InAutomatedTest->GetAdditionalCommandLine().IsEmpty())
+	{
+		FString AdditionalCommandLine = InAutomatedTest->GetAdditionalCommandLine();
+		LauncherProfileManager->OnPostProcessAutomatedTestCommandLine().Broadcast(InAutomatedTest, InProfile, AdditionalCommandLine);
+
+		if (!AdditionalCommandLine.IsEmpty())
+		{
+			AutomatedTestCommand += TEXT(" ");
+			AutomatedTestCommand += AdditionalCommandLine;
+		}
+	}
+
+	return AutomatedTestCommand;
+}
+
 FString FLauncherWorker::MakeBuildCookRunParamsForProjectCustomBuild(const ILauncherProfileRef& InProfile, const TArray<FString>& InPlatforms)
 {
 	// get the basic UAT command line
@@ -1243,7 +1356,7 @@ void FLauncherWorker::CreateAndExecuteTasks( const ILauncherProfileRef& InProfil
 		class FWaitForCookInEditorToFinish : public FLauncherTask
 		{
 		public:
-			FWaitForCookInEditorToFinish() : FLauncherTask( FString(TEXT("Cooking in the editor")), FString(TEXT("Prepairing content to run on device")))
+			FWaitForCookInEditorToFinish() : FLauncherTask( FString(TEXT("Cooking in the editor")), FString(TEXT("Preparing content to run on device")))
 			{
 			}
 			virtual bool PerformTask( FLauncherTaskChainState& ChainState ) override
@@ -1274,7 +1387,21 @@ void FLauncherWorker::CreateAndExecuteTasks( const ILauncherProfileRef& InProfil
 		TurnkeyCommand = FString::Printf(TEXT("Turnkey -command=VerifySdk -type=Flash -device=%s -UpdateIfNeeded -utf8output -WaitForUATMutex %s"), *FString::Join(DeviceGroup->GetDeviceIDs(), TEXT("+")), *ITurnkeyIOModule::Get().GetUATParams());
 	}
 
-	TSharedPtr<FLauncherTask> LaunchTask = MakeShareable(new FLauncherUATTask(UATCommand, TEXT("Launch Task"), TEXT("Launching UAT..."), ReadPipe, WritePipe, InProfile->GetEditorExe(), ProcHandle, this, StartString, TurnkeyCommand));
+	// have Automated tests
+	FString PostCommand = CreateUATCommandForAutomatedTests( InProfile, Platforms, Commands, StartString );
+
+	// wait for completion of UAT
+	{
+		FCommandDesc Desc;
+		FText Command = LOCTEXT("LauncherCompletionDesc", "UAT post launch cleanup");
+		Desc.Name = "Post Launch Task";
+		Desc.Desc = Command.ToString();
+		Desc.EndText = TEXT("********** LAUNCH COMPLETED **********");
+		Commands.Add(Desc);
+	}
+
+
+	TSharedPtr<FLauncherTask> LaunchTask = MakeShareable(new FLauncherUATTask(UATCommand, TEXT("Launch Task"), TEXT("Launching UAT..."), ReadPipe, WritePipe, InProfile->GetEditorExe(), ProcHandle, this, StartString, TurnkeyCommand, PostCommand));
 	AddTask(LaunchTask);
 	for (int32 Index = 0; Index < Commands.Num(); ++Index)
 	{
diff --git a/Engine/Source/Developer/LauncherServices/Private/Launcher/LauncherWorker.h b/Engine/Source/Developer/LauncherServices/Private/Launcher/LauncherWorker.h
--- a/Engine/Source/Developer/LauncherServices/Private/Launcher/LauncherWorker.h
+++ b/Engine/Source/Developer/LauncherServices/Private/Launcher/LauncherWorker.h
@@ -125,6 +125,9 @@ protected:
 
 	FString CreateUATCommand( const ILauncherProfileRef& InProfile, const TArray<FString>& InPlatforms, TArray<FCommandDesc>& OutCommands, FString& CommandStart, bool bForTurnkeyCustomBuild = false );
 
+	FString CreateUATCommandForAutomatedTests( const ILauncherProfileRef& InProfile, const TArray<FString>& InPlatforms, TArray<FCommandDesc>& OutCommands, FString& OutCommandStart );
+	FString CreateUATCommandForAutomatedTest( const ILauncherProfileAutomatedTestRef& InAutomatedTest, const ILauncherProfileRef& InProfile, const TArray<FString>& InPlatforms, TArray<FCommandDesc>& OutCommands, FString& OutCommandStart );
+
 	//start listening to the device discovery routine
 	void EnableDeviceDiscoveryListener();
 
diff --git a/Engine/Source/Developer/LauncherServices/Private/Profiles/LauncherProfile.h b/Engine/Source/Developer/LauncherServices/Private/Profiles/LauncherProfile.h
--- a/Engine/Source/Developer/LauncherServices/Private/Profiles/LauncherProfile.h
+++ b/Engine/Source/Developer/LauncherServices/Private/Profiles/LauncherProfile.h
@@ -18,6 +18,7 @@
 #include "ITargetDeviceServicesModule.h"
 #include "GameProjectHelper.h"
 #include "Profiles/LauncherProfileLaunchRole.h"
+#include "Profiles/LauncherProfileAutomatedTest.h"
 #include "PlatformInfo.h"
 #include "TargetReceipt.h"
 #include "DesktopPlatformModule.h"
@@ -60,6 +61,7 @@ enum ELauncherVersion
 	LAUNCHERSERVICES_ADDBUILDARCHITECTURE = 38,
 	LAUNCHERSERVICES_ADDCUSTOMPROPERTIES = 39,
 	LAUNCHERSERVICES_INCREMENTALCOOKMODE = 40,
+	LAUNCHERSERVICES_AUTOMATEDTESTS = 41,
 	//ADD NEW STUFF HERE
 
 
@@ -1155,6 +1157,31 @@ public:
 			IncrementalCookMode = CookIncremental ? ELauncherProfileIncrementalCookMode::ModifiedOnly : ELauncherProfileIncrementalCookMode::None;
 		}
 
+		if (Version >= LAUNCHERSERVICES_AUTOMATEDTESTS)
+		{
+			// serialize automated tests
+			if (Archive.IsLoading())
+			{
+				AutomatedTests.Reset();
+			}
+			int NumAutomatedTests = AutomatedTests.Num();
+			Archive << NumAutomatedTests;
+			for( int32 TestIndex = 0; TestIndex < NumAutomatedTests; TestIndex++)
+			{
+				if (Archive.IsLoading())
+				{
+					AutomatedTests.Add(MakeShared<FLauncherProfileAutomatedTest>(Archive));
+				}
+				else
+				{
+					AutomatedTests[TestIndex]->Serialize(Archive);
+				}
+			}
+
+			Archive << AutomatedTestBuildPath;
+			Archive << bWantAutomatedTestBuild;
+		}
+
 		DefaultLaunchRole->Serialize(Archive);
 
 		// serialize launch roles
@@ -1324,6 +1351,17 @@ public:
 			Writer.WriteArrayEnd();
 		}
 
+		// serialize the automated tests
+		Writer.WriteArrayStart("AutomatedTests");
+		for ( ILauncherProfileAutomatedTestRef AutomatedTest : AutomatedTests)
+		{
+			AutomatedTest->Save(Writer);
+		}
+		Writer.WriteArrayEnd();
+		Writer.WriteValue("AutomatedTestBuildPath", AutomatedTestBuildPath);
+		Writer.WriteValue("WantAutomatedTestBuild", bWantAutomatedTestBuild);
+
+
 		// write out the UAT project params
 		SaveUATParams(Writer);
 		Writer.WriteObjectEnd();
@@ -2077,6 +2115,21 @@ public:
 			IncrementalCookMode = CookIncremental ? ELauncherProfileIncrementalCookMode::ModifiedOnly : ELauncherProfileIncrementalCookMode::None;
 		}
 
+		if (Version >= LAUNCHERSERVICES_AUTOMATEDTESTS)
+		{
+			AutomatedTests.Reset();
+			const TArray<TSharedPtr<FJsonValue>>* Values = nullptr;
+			if (Object.TryGetArrayField(TEXT("AutomatedTests"), Values))
+			{
+				for (TSharedPtr<FJsonValue> Value : *Values)
+				{
+					AutomatedTests.Add(MakeShared<FLauncherProfileAutomatedTest>(*Value->AsObject().Get()));
+				}
+			}
+			AutomatedTestBuildPath = Object.GetStringField(TEXT("AutomatedTestBuildPath"));
+			bWantAutomatedTestBuild = Object.GetBoolField(TEXT("WantAutomatedTestBuild"));
+		}
+
 
 		// load the default launch role
 		TSharedPtr<FJsonObject> Role = Object.GetObjectField(TEXT("DefaultRole"));
@@ -2231,6 +2284,10 @@ public:
 		bIsDeviceASimulator = false;
 		bMakeBinaryConfig = false;
 
+		AutomatedTests.Reset();
+		AutomatedTestBuildPath.Reset();
+		bWantAutomatedTestBuild = false;
+
 		RefreshValidBuildTargets();
 		Validate();
 	}
@@ -2892,6 +2949,59 @@ public:
 		return CustomBoolProperties;
 	}
 
+	virtual const TArray<ILauncherProfileAutomatedTestRef>& GetAutomatedTests() const override
+	{
+		return AutomatedTests;
+	}
+
+	virtual ILauncherProfileAutomatedTestPtr GetAutomatedTest( const TCHAR* InternalName ) const override
+	{
+		const ILauncherProfileAutomatedTestRef* TestPtr = AutomatedTests.FindByPredicate([InternalName]( const ILauncherProfileAutomatedTestRef& Test )
+		{
+			return FCString::Stricmp( InternalName, Test->GetInternalName()) == 0;
+		});
+
+		return TestPtr ? TestPtr->ToSharedPtr() : nullptr;
+	}
+
+	virtual ILauncherProfileAutomatedTestRef FindOrAddAutomatedTest( const TCHAR* InternalName ) override
+	{
+		ILauncherProfileAutomatedTestPtr AutomatedTest = GetAutomatedTest(InternalName);
+		if (AutomatedTest == nullptr)
+		{
+			AutomatedTest = MakeShared<FLauncherProfileAutomatedTest>(InternalName);
+			AutomatedTests.Add(AutomatedTest.ToSharedRef());
+		}
+
+		return AutomatedTest.ToSharedRef();
+	}
+
+	virtual void RemoveAutomatedTest( const TCHAR* InternalName ) override
+	{
+		ILauncherProfileAutomatedTestPtr AutomatedTest = GetAutomatedTest(InternalName);
+		if (AutomatedTest != nullptr)
+		{
+			AutomatedTests.Remove(AutomatedTest.ToSharedRef());
+		}
+	}
+
+	virtual void SetAutomatedTestBuildPath( const FString& InAutomatedTestBuildPath )
+	{
+		AutomatedTestBuildPath = InAutomatedTestBuildPath;
+	}
+	virtual const FString& GetAutomatedTestBuildPath() const
+	{
+		return AutomatedTestBuildPath;
+	}
+
+	virtual void SetIsUsingAutomatedTestBuild( bool bInWantAutomatedTestBuild ) override
+	{
+		bWantAutomatedTestBuild = bInWantAutomatedTestBuild;
+	}
+	virtual bool IsUsingAutomatedTestBuild() const override
+	{
+		return bWantAutomatedTestBuild;
+	}
 
 	//~ End ILauncherProfile Interface
 
@@ -3521,6 +3631,11 @@ private:
 	TMap<FString,FString> CustomStringProperties;
 	TMap<FString, bool>   CustomBoolProperties;
 
+	// Automated tests
+	TArray<ILauncherProfileAutomatedTestRef> AutomatedTests;
+	FString AutomatedTestBuildPath;
+	bool bWantAutomatedTestBuild;
+
 private:
 
 	// cook in the editor callbacks (not valid for any other cook mode)
diff --git a/Engine/Source/Developer/LauncherServices/Private/Profiles/LauncherProfileAutomatedTest.h b/Engine/Source/Developer/LauncherServices/Private/Profiles/LauncherProfileAutomatedTest.h
new file mode 100644
--- /dev/null
+++ b/Engine/Source/Developer/LauncherServices/Private/Profiles/LauncherProfileAutomatedTest.h
@@ -0,0 +1,123 @@
+// Copyright Epic Games, Inc. All Rights Reserved.
+
+#pragma once
+
+#include "CoreMinimal.h"
+#include "ILauncherProfileAutomatedTest.h"
+#include "Dom/JsonObject.h"
+
+class FLauncherProfileAutomatedTest final : public ILauncherProfileAutomatedTest
+{
+public:
+	FLauncherProfileAutomatedTest( const TCHAR* InInternalName ) 
+		: InternalName(InInternalName)
+	{
+	}
+
+	FLauncherProfileAutomatedTest(const FJsonObject& Object)
+	{
+		Load(Object);
+	}
+
+	FLauncherProfileAutomatedTest( FArchive& Archive )
+	{
+		Serialize(Archive);
+	}
+
+	virtual void SetUATCommand( const TCHAR* InUATCommand ) override
+	{
+		UATCommand = InUATCommand;
+	}
+
+	virtual const FString& GetUATCommand() const override
+	{
+		return UATCommand;
+	}
+
+	virtual void SetTests( const TCHAR* InTests ) override
+	{
+		Tests = InTests;
+	}
+	virtual const FString& GetTests() const override
+	{
+		return Tests;
+	}
+
+	virtual void SetAdditionalCommandLine( const TCHAR* InAdditionalCommandLine ) override
+	{
+		AdditionalCommandLine = InAdditionalCommandLine;
+	}
+	virtual const FString& GetAdditionalCommandLine() const override
+	{
+		return AdditionalCommandLine;
+	}
+
+	virtual void SetEnabled( bool bInEnabled ) override
+	{
+		bEnabled = bInEnabled;
+	}
+	virtual bool IsEnabled() const override
+	{
+		return bEnabled;
+	}
+
+	virtual void SetPriority( int32 InPriority ) override
+	{
+		Priority = InPriority;
+	}
+	virtual int32 GetPriority() const override
+	{
+		return Priority;
+	}
+
+	virtual const TCHAR* GetInternalName() const override
+	{
+		return *InternalName;
+	}
+
+	virtual void Load(const FJsonObject& Object) override
+	{
+		InternalName          = Object.GetStringField(TEXT("InternalName"));
+		UATCommand            = Object.GetStringField(TEXT("UATCommand"));
+		Tests                 = Object.GetStringField(TEXT("Tests"));
+		AdditionalCommandLine = Object.GetStringField(TEXT("AdditionalCommandLine"));
+		bEnabled              = Object.GetBoolField(TEXT("Enabled"));
+		Priority              = Object.GetIntegerField(TEXT("Priority"));
+
+	}
+	virtual void Save(TJsonWriter<>& Writer) override
+	{
+		Writer.WriteObjectStart();
+		Writer.WriteValue(TEXT("InternalName"),          InternalName);
+		Writer.WriteValue(TEXT("UATCommand"),            UATCommand);
+		Writer.WriteValue(TEXT("Tests"),                 Tests);
+		Writer.WriteValue(TEXT("AdditionalCommandLine"), AdditionalCommandLine);
+		Writer.WriteValue(TEXT("Enabled"),               bEnabled);
+		Writer.WriteValue(TEXT("Priority"),              FString::FromInt(Priority));
+		Writer.WriteObjectEnd();
+	}
+
+	virtual void Serialize( FArchive& Archive ) override final
+	{
+		Archive << InternalName
+				<< UATCommand
+				<< Tests
+				<< AdditionalCommandLine
+				<< bEnabled
+				<< Priority
+		;
+	}
+
+
+protected:
+	FString UATCommand = TEXT("RunUnreal");
+	FString Tests; //-Test=xxx,yyy
+	FString AdditionalCommandLine;
+
+	bool bEnabled = true;
+	int32 Priority = 0; // longer tests should generally have lower priority
+
+	FString InternalName;
+};
+
+
diff --git a/Engine/Source/Developer/LauncherServices/Private/Profiles/LauncherProfileManager.h b/Engine/Source/Developer/LauncherServices/Private/Profiles/LauncherProfileManager.h
--- a/Engine/Source/Developer/LauncherServices/Private/Profiles/LauncherProfileManager.h
+++ b/Engine/Source/Developer/LauncherServices/Private/Profiles/LauncherProfileManager.h
@@ -90,6 +90,11 @@ public:
 		return ProcessProcessLaunchCommandLine;
 	}
 
+	virtual FOnLauncherProfilePostProcessAutomatedTestCommandLine& OnPostProcessAutomatedTestCommandLine( ) override
+	{
+		return PostProcessAutomatedTestCommandLine;
+	}
+
 	virtual void RemoveDeviceGroup( const ILauncherDeviceGroupRef& DeviceGroup ) override;
 
 	virtual void SaveDeviceGroups() override;
@@ -219,5 +224,7 @@ private:
 	// Holds a delegate to be invoked to modify the command line before launching
 	FOnLauncherProfilePostProcessCommandLine ProcessProcessLaunchCommandLine;
 
+	// Holds a delegate to be involked to modify the command line for an automated test before launching
+	FOnLauncherProfilePostProcessAutomatedTestCommandLine PostProcessAutomatedTestCommandLine;
 
 };
diff --git a/Engine/Source/Developer/LauncherServices/Public/ILauncherProfile.h b/Engine/Source/Developer/LauncherServices/Public/ILauncherProfile.h
--- a/Engine/Source/Developer/LauncherServices/Public/ILauncherProfile.h
+++ b/Engine/Source/Developer/LauncherServices/Public/ILauncherProfile.h
@@ -7,6 +7,7 @@
 #include "Containers/Map.h"
 #include "ILauncherDeviceGroup.h"
 #include "ILauncherProfileLaunchRole.h"
+#include "ILauncherProfileAutomatedTest.h"
 
 class Error;
 class FJsonObject;
@@ -1526,6 +1527,48 @@ public:
 	 */
 	virtual const TArray<FString>& GetEditorArchitectures() const = 0;
 
+
+	/**
+	 * Returns a list of all automated tests in the profile
+	 */
+	virtual const TArray<ILauncherProfileAutomatedTestRef>& GetAutomatedTests() const = 0;
+
+	/**
+	 * Returns an existing automated test by internal name
+	 */
+	virtual ILauncherProfileAutomatedTestPtr GetAutomatedTest( const TCHAR* InternalName ) const = 0;
+
+	/**
+	 * Returns an automated test by internal name, creating it if it doesn't already exist
+	 */
+	virtual ILauncherProfileAutomatedTestRef FindOrAddAutomatedTest( const TCHAR* InternalName ) = 0;
+
+	/**
+	 * Removes an automated test by internal name
+	 */
+	virtual void RemoveAutomatedTest( const TCHAR* InternalName ) = 0;
+
+
+	/**
+	 * Sets the build to use for ALL automated tests, equivalent to the -build parameter on RunUnreal etc. See also SetIsUsingAutomatedTestBuild.
+	 */
+	virtual void SetAutomatedTestBuildPath( const FString& AutomatedTestBuildPath ) = 0;
+
+	/**
+	 * Returns the current automated test build path
+	 */
+	virtual const FString& GetAutomatedTestBuildPath() const = 0;
+
+	/**
+	 * Set whether to use the specified automated test build or the current launch build. Default is to use the current build.
+	 */
+	virtual void SetIsUsingAutomatedTestBuild( bool bWantAutomatedTestBuild ) = 0;
+
+	/**
+	 * Returns whether we are using the specified automated test build or not.
+	 */
+	virtual bool IsUsingAutomatedTestBuild() const = 0;
+
 public:
 	/**
 	 * Helper function to get all of the build targets available for this profile, based on its current project & cook platforms
diff --git a/Engine/Source/Developer/LauncherServices/Public/ILauncherProfileAutomatedTest.h b/Engine/Source/Developer/LauncherServices/Public/ILauncherProfileAutomatedTest.h
new file mode 100644
--- /dev/null
+++ b/Engine/Source/Developer/LauncherServices/Public/ILauncherProfileAutomatedTest.h
@@ -0,0 +1,46 @@
+// Copyright Epic Games, Inc. All Rights Reserved.
+
+#pragma once
+
+#include "CoreMinimal.h"
+#include "Serialization/JsonWriter.h"
+
+class FJsonObject;
+class ILauncherProfileLaunchRole;
+
+/**
+ * Automated test parameters
+ */
+class ILauncherProfileAutomatedTest
+{
+public:
+	//@todo: comments
+
+	// recommendation for InternalName is is "ExtensionName.InternalName"
+	virtual void SetUATCommand( const TCHAR* UATCommand ) = 0;
+	virtual const FString& GetUATCommand() const = 0;
+
+	virtual void SetTests( const TCHAR* Tests ) = 0; //-Test=xxx,yyy
+	virtual const FString& GetTests() const = 0;
+
+	virtual void SetAdditionalCommandLine( const TCHAR* AdditionalCommandLine ) = 0;
+	virtual const FString& GetAdditionalCommandLine() const = 0;
+
+	virtual void SetEnabled( bool bEnabled ) = 0;
+	virtual bool IsEnabled() const = 0;
+
+	virtual void SetPriority( int32 Priority ) = 0; // longer tests should generally have lower priority
+	virtual int32 GetPriority() const = 0;
+
+	virtual const TCHAR* GetInternalName() const = 0;
+
+	virtual void Load(const FJsonObject& Object) = 0;
+	virtual void Save(TJsonWriter<>& Writer) = 0;
+	virtual void Serialize( FArchive& Archive ) = 0;
+
+
+	virtual ~ILauncherProfileAutomatedTest() = default;
+};
+
+typedef TSharedPtr<ILauncherProfileAutomatedTest> ILauncherProfileAutomatedTestPtr;
+typedef TSharedRef<ILauncherProfileAutomatedTest> ILauncherProfileAutomatedTestRef;
diff --git a/Engine/Source/Developer/LauncherServices/Public/ILauncherProfileManager.h b/Engine/Source/Developer/LauncherServices/Public/ILauncherProfileManager.h
--- a/Engine/Source/Developer/LauncherServices/Public/ILauncherProfileManager.h
+++ b/Engine/Source/Developer/LauncherServices/Public/ILauncherProfileManager.h
@@ -57,6 +57,13 @@ DECLARE_MULTICAST_DELEGATE(FOnLauncherProfileManagerProjectChanged);
  */
 DECLARE_MULTICAST_DELEGATE_TwoParams(FOnLauncherProfilePostProcessCommandLine, const ILauncherProfileRef&, FString& );
 
+/**
+ * Declares a delegate to be invoked to post-process a command line before launching and automated test
+ *
+ * The first parameter is the automated test, the second is the profile that is being launched, the third parameter is the command line to modify.
+ */
+DECLARE_MULTICAST_DELEGATE_ThreeParams(FOnLauncherProfilePostProcessAutomatedTestCommandLine, const ILauncherProfileAutomatedTestRef&, const ILauncherProfileRef&, FString& );
+
 
 /** Type definition for shared references to instances of ILauncherProfile. */
 typedef TSharedPtr<class ILauncherProfileWizard> ILauncherProfileWizardPtr;
@@ -415,6 +422,12 @@ public:
 	 */
 	virtual FOnLauncherProfilePostProcessCommandLine& OnPostProcessLaunchCommandLine( ) = 0;
 
+	/**
+	 * Returns a delegate that is invoked to post-process the command line before an automated test is launched
+	 * 
+	 * @return The delegate.
+	 */
+	virtual FOnLauncherProfilePostProcessAutomatedTestCommandLine& OnPostProcessAutomatedTestCommandLine( ) = 0;
 
 public:
 
