static analysis fixes

#rnx

diff --git a/Engine/Plugins/Developer/ProjectLauncher/Source/CommonLaunchExtensions/Private/BootTest/BootTestLaunchExtension.cpp b/Engine/Plugins/Developer/ProjectLauncher/Source/CommonLaunchExtensions/Private/BootTest/BootTestLaunchExtension.cpp
--- a/Engine/Plugins/Developer/ProjectLauncher/Source/CommonLaunchExtensions/Private/BootTest/BootTestLaunchExtension.cpp
+++ b/Engine/Plugins/Developer/ProjectLauncher/Source/CommonLaunchExtensions/Private/BootTest/BootTestLaunchExtension.cpp
@@ -12,7 +12,7 @@ void FBootTestLaunchExtensionInstance::CustomizeTree( ProjectLauncher::FLaunchPr
 {
 	// note: Windowed option is here mostly as an example of adding UI items for a particular profile - your tests may not need it
 
-	auto IsVisible = [this]()
+	auto IsVisible = [this]() // this always outlives the ProfileTreeData
 	{
 		return IsTestActive();
 	};
diff --git a/Engine/Source/Developer/LauncherServices/Private/Launcher/LauncherWorker.cpp b/Engine/Source/Developer/LauncherServices/Private/Launcher/LauncherWorker.cpp
--- a/Engine/Source/Developer/LauncherServices/Private/Launcher/LauncherWorker.cpp
+++ b/Engine/Source/Developer/LauncherServices/Private/Launcher/LauncherWorker.cpp
@@ -1062,7 +1062,7 @@ FString FLauncherWorker::CreateUATCommand( const ILauncherProfileRef& InProfile,
 		}
 
 		// run
-		if (InProfile->GetLaunchMode() != ELauncherProfileLaunchModes::DoNotLaunch && Profile->GetAutomatedTests().Num() == 0)
+		if (InProfile->GetLaunchMode() != ELauncherProfileLaunchModes::DoNotLaunch && InProfile->GetAutomatedTests().Num() == 0)
 		{
 			UATCommand += TEXT(" -run ");
 
@@ -1130,9 +1130,9 @@ FString FLauncherWorker::CreateUATCommandForAutomatedTests( const ILauncherProfi
 {
 	FString UATCommand;
 
-	if (Profile->GetAutomatedTests().Num() > 0)
+	if (InProfile->GetAutomatedTests().Num() > 0)
 	{
-		TArray<ILauncherProfileAutomatedTestRef> PrioritizedAutomatedTests = Profile->GetAutomatedTests();
+		TArray<ILauncherProfileAutomatedTestRef> PrioritizedAutomatedTests = InProfile->GetAutomatedTests();
 		PrioritizedAutomatedTests.Sort([]( const ILauncherProfileAutomatedTestRef& A, const ILauncherProfileAutomatedTestRef& B)
 		{
 			return A->GetPriority() < B->GetPriority();
@@ -1231,16 +1231,13 @@ FString FLauncherWorker::CreateUATCommandForAutomatedTest( const ILauncherProfil
 	}
 
 	// final extra command line parameters
-	if (InAutomatedTest->GetAdditionalCommandLine().IsEmpty())
-	{
-		FString AdditionalCommandLine = InAutomatedTest->GetAdditionalCommandLine();
-		LauncherProfileManager->OnPostProcessAutomatedTestCommandLine().Broadcast(InAutomatedTest, InProfile, AdditionalCommandLine);
+	FString AdditionalCommandLine = InAutomatedTest->GetAdditionalCommandLine();
+	LauncherProfileManager->OnPostProcessAutomatedTestCommandLine().Broadcast(InAutomatedTest, InProfile, AdditionalCommandLine);
 
-		if (!AdditionalCommandLine.IsEmpty())
-		{
-			AutomatedTestCommand += TEXT(" ");
-			AutomatedTestCommand += AdditionalCommandLine;
-		}
+	if (!AdditionalCommandLine.IsEmpty())
+	{
+		AutomatedTestCommand += TEXT(" ");
+		AutomatedTestCommand += AdditionalCommandLine;
 	}
 
 	return AutomatedTestCommand;
diff --git a/Engine/Source/Developer/LauncherServices/Private/Profiles/LauncherProfileAutomatedTest.h b/Engine/Source/Developer/LauncherServices/Private/Profiles/LauncherProfileAutomatedTest.h
--- a/Engine/Source/Developer/LauncherServices/Private/Profiles/LauncherProfileAutomatedTest.h
+++ b/Engine/Source/Developer/LauncherServices/Private/Profiles/LauncherProfileAutomatedTest.h
@@ -93,7 +93,7 @@ public:
 		Writer.WriteValue(TEXT("Tests"),                 Tests);
 		Writer.WriteValue(TEXT("AdditionalCommandLine"), AdditionalCommandLine);
 		Writer.WriteValue(TEXT("Enabled"),               bEnabled);
-		Writer.WriteValue(TEXT("Priority"),              FString::FromInt(Priority));
+		Writer.WriteValue(TEXT("Priority"),              Priority);
 		Writer.WriteObjectEnd();
 	}
 
