Bug Fix - Always stage the build when Pak Files are required.

#rnx

diff --git a/Engine/Source/Developer/LauncherServices/Private/Launcher/LauncherWorker.cpp b/Engine/Source/Developer/LauncherServices/Private/Launcher/LauncherWorker.cpp
--- a/Engine/Source/Developer/LauncherServices/Private/Launcher/LauncherWorker.cpp
+++ b/Engine/Source/Developer/LauncherServices/Private/Launcher/LauncherWorker.cpp
@@ -771,6 +771,7 @@ FString FLauncherWorker::CreateUATCommand( const ILauncherProfileRef& InProfile,
 			if (InProfile->IsPackingWithUnrealPak())
 			{
 				UATCommand += TEXT(" -pak");
+				UATCommand += TEXT(" -stage");
 				if (InProfile->IsUsingIoStore())
 				{
 					UATCommand += TEXT(" -iostore");
