Fix Project Launcher 2 failing to deploy Automated Test

#tests Launch Lyra Automated boot test via PL2

diff --git a/Engine/Source/Developer/LauncherServices/Private/Launcher/LauncherWorker.cpp b/Engine/Source/Developer/LauncherServices/Private/Launcher/LauncherWorker.cpp
--- a/Engine/Source/Developer/LauncherServices/Private/Launcher/LauncherWorker.cpp
+++ b/Engine/Source/Developer/LauncherServices/Private/Launcher/LauncherWorker.cpp
@@ -1217,6 +1217,20 @@ FString FLauncherWorker::CreateUATCommandForAutomatedTest( const ILauncherProfil
 		AutomatedTestCommand += TEXT("-dev ");
 	}
 
+	// Device ID passed down from Project Launcher can contain
+	// "Windows" instead of "Win64" in the Device Id, which is 
+	// not recognized by "RunUnreal" UAT command. In this case,
+	// we convert the Device Id as required. 
+	const auto ConvertDeviceID = [](FString DeviceID) -> FString
+	{
+		if (DeviceID.Contains("Windows"))
+		{
+			DeviceID = DeviceID.Replace(TEXT("Windows"), TEXT("Win64"));
+		}
+
+		return DeviceID.TrimStartAndEnd();
+	};
+
 	// specify the devices to use
 	ILauncherDeviceGroupPtr DeviceGroup = InProfile->GetDeployedDeviceGroup();
 	if (DeviceGroup.IsValid() && DeviceGroup->GetNumDevices() > 0)
@@ -1227,7 +1241,7 @@ FString FLauncherWorker::CreateUATCommandForAutomatedTest( const ILauncherProfil
 			FString FinalDeviceId = DeviceId.TrimStartAndEnd().Replace(TEXT("@"), TEXT(":"));
 			if (!FinalDeviceId.IsEmpty())
 			{
-				DeviceNames += TEXT(",\"") + FinalDeviceId + TEXT("\"");
+				DeviceNames += TEXT(",\"") + ConvertDeviceID(FinalDeviceId) + TEXT("\"");
 			}
 		}
 		if (!DeviceNames.IsEmpty())
