Project Launcher changes for incremental cook:
-  Restore 'incremental cook' tickbox instead of the 'incremental cook modes' option
- Profile verification checks the new 'legacy' boolean to determine if unversioned cooks are allowed
- Exposing unversioned cook option in Project Launcher 2

#rnx

diff --git a/Engine/Plugins/Developer/ProjectLauncher/Source/ProjectLauncher/Private/ProfileTree/CustomProfileTreeBuilder.cpp b/Engine/Plugins/Developer/ProjectLauncher/Source/ProjectLauncher/Private/ProfileTree/CustomProfileTreeBuilder.cpp
--- a/Engine/Plugins/Developer/ProjectLauncher/Source/ProjectLauncher/Private/ProfileTree/CustomProfileTreeBuilder.cpp
+++ b/Engine/Plugins/Developer/ProjectLauncher/Source/ProjectLauncher/Private/ProfileTree/CustomProfileTreeBuilder.cpp
@@ -34,6 +34,7 @@ namespace ProjectLauncher
 		FLaunchProfileTreeNode& MapsAndCookingHeader = TreeData->AddHeading( TEXT("Cooking"), LOCTEXT("CookingHeading", "Maps And Cooking") );
 		AddCookProperty(MapsAndCookingHeader);
 		AddIncrementalCookProperty(MapsAndCookingHeader);
+		AddUnversionedCookProperty(MapsAndCookingHeader);
 		AddMapsToCookProperty(MapsAndCookingHeader);
 		AddAdditionalCookerOptionsProperty(MapsAndCookingHeader);
 
diff --git a/Engine/Plugins/Developer/ProjectLauncher/Source/ProjectLauncher/Private/ProfileTree/GenericProfileTreeBuilder.cpp b/Engine/Plugins/Developer/ProjectLauncher/Source/ProjectLauncher/Private/ProfileTree/GenericProfileTreeBuilder.cpp
--- a/Engine/Plugins/Developer/ProjectLauncher/Source/ProjectLauncher/Private/ProfileTree/GenericProfileTreeBuilder.cpp
+++ b/Engine/Plugins/Developer/ProjectLauncher/Source/ProjectLauncher/Private/ProfileTree/GenericProfileTreeBuilder.cpp
@@ -333,47 +333,27 @@ namespace ProjectLauncher
 
 	void FGenericProfileTreeBuilder::AddIncrementalCookProperty( FLaunchProfileTreeNode& HeadingNode )
 	{
-		auto GetDisplayName = []( ELauncherProfileIncrementalCookMode::Type Mode)
-		{
-			switch (Mode)
+		HeadingNode.AddBoolean( LOCTEXT("IncrementalCookLabel", "Incremental Cook"),
 			{
-				case ELauncherProfileIncrementalCookMode::None:                    return LOCTEXT("IncrementalCookNone",                 "None");
-				case ELauncherProfileIncrementalCookMode::ModifiedOnly:            return LOCTEXT("IncrementalCookModified",             "Modified Only (legacy)");
-				case ELauncherProfileIncrementalCookMode::ModifiedAndDependencies: return LOCTEXT("IncrementalCookModifiedDependencies", "Modified & Dependencies (recommended)");
-			}
-			return FText::GetEmpty();
-		};
-
-		auto GetToolTip = []( ELauncherProfileIncrementalCookMode::Type Mode)
-		{
-			switch (Mode)
-			{
-				case ELauncherProfileIncrementalCookMode::None:                    return LOCTEXT("IncrementalCookNoneTip",                 "This will try to cook everything");
-				case ELauncherProfileIncrementalCookMode::ModifiedOnly:            return LOCTEXT("IncrementalCookModifiedTip",             "This will only try to cook any modified assets but won't try to cook anything that depends on these assets. This is the old, legacy option and is faster but unreliable");
-				case ELauncherProfileIncrementalCookMode::ModifiedAndDependencies: return LOCTEXT("IncrementalCookModifiedDependenciesTip", "This will try to cook any modified assets and those assets that depend on them. This is the new method and is much more reliable but is a little slower");
+				.GetValue = [this]()			{ return Profile->IsCookingIncrementally(); },
+				.SetValue = [this](bool bValue)	{ Profile->SetIncrementalCooking(bValue); },
+				.GetDefaultValue = [this]()		{ return DefaultProfile->IsCookingIncrementally(); },
+				.IsVisible = ForEnabledCooked,
+				.Validation = FValidation({ELauncherProfileValidationErrors::UnversionedAndIncremental}),
 			}
-			return FText::GetEmpty();
-		};
+		);
+	}
 
-		HeadingNode.AddWidget( LOCTEXT("IncrementalCookLabel", "Incremental Cook"),
+	void FGenericProfileTreeBuilder::AddUnversionedCookProperty( FLaunchProfileTreeNode& HeadingNode )
+	{
+		HeadingNode.AddBoolean( LOCTEXT("UnversionedCookLabel", "Unversioned Cook"),
 			{
-				.IsDefault = [this]()		{ return GetIncrementalCookMode() == DefaultProfile->GetIncrementalCookMode(); },
-				.SetToDefault = [this]()	{ SetIncrementalCookMode(DefaultProfile->GetIncrementalCookMode()); },
+				.GetValue = [this]()			{ return Profile->IsCookingUnversioned(); },
+				.SetValue = [this](bool bValue)	{ Profile->SetUnversionedCooking(bValue); },
+				.GetDefaultValue = [this]()		{ return DefaultProfile->IsCookingUnversioned(); },
 				.IsVisible = ForEnabledCooked,
 				.Validation = FValidation({ELauncherProfileValidationErrors::UnversionedAndIncremental}),
-			},
-			SNew(SCustomLaunchCombo<ELauncherProfileIncrementalCookMode::Type>)
-			.OnSelectionChanged(this, &FGenericProfileTreeBuilder::SetIncrementalCookMode)
-			.SelectedItem(this, &FGenericProfileTreeBuilder::GetIncrementalCookMode)
-			.GetDisplayName_Lambda(GetDisplayName)
-			.GetItemToolTip_Lambda(GetToolTip)
-			.Items( TArray<ELauncherProfileIncrementalCookMode::Type>(
-				{
-					ELauncherProfileIncrementalCookMode::None,
-					ELauncherProfileIncrementalCookMode::ModifiedOnly,
-					ELauncherProfileIncrementalCookMode::ModifiedAndDependencies,
-				})
-			)
+			}
 		);
 	}
 
@@ -1045,22 +1025,6 @@ namespace ProjectLauncher
 		return (InProfile->GetCookMode() != ELauncherProfileCookModes::DoNotCook);
 	}
 
-
-	void FGenericProfileTreeBuilder::SetIncrementalCookMode( ELauncherProfileIncrementalCookMode::Type Mode )
-	{
-		Profile->SetIncrementalCookMode(Mode);
-		// should always use unversioned, except the LegacyIterative cook (-iterate commandline argument) because it does not handle invalidation due to c++ changes.
-		Profile->SetUnversionedCooking(Mode != ELauncherProfileIncrementalCookMode::ModifiedOnly);
-		OnPropertyChanged();
-	}
-
-	ELauncherProfileIncrementalCookMode::Type FGenericProfileTreeBuilder::GetIncrementalCookMode() const
-	{
-		return Profile->GetIncrementalCookMode();
-	}
-
-
-
 	void FGenericProfileTreeBuilder::SetMapsToCook(TArray<FString> MapsToCook)
 	{
 		Profile->ClearCookedMaps(); 
diff --git a/Engine/Plugins/Developer/ProjectLauncher/Source/ProjectLauncher/Public/ProfileTree/GenericProfileTreeBuilder.h b/Engine/Plugins/Developer/ProjectLauncher/Source/ProjectLauncher/Public/ProfileTree/GenericProfileTreeBuilder.h
--- a/Engine/Plugins/Developer/ProjectLauncher/Source/ProjectLauncher/Public/ProfileTree/GenericProfileTreeBuilder.h
+++ b/Engine/Plugins/Developer/ProjectLauncher/Source/ProjectLauncher/Public/ProfileTree/GenericProfileTreeBuilder.h
@@ -56,6 +56,7 @@ namespace ProjectLauncher
 		UE_API void AddImportZenSnapshotProperty( FLaunchProfileTreeNode& HeadingNode );
 		UE_API void AddZenPakStreamingPathProperty( FLaunchProfileTreeNode& HeadingNode );
 		UE_API void AddIncrementalCookProperty( FLaunchProfileTreeNode& HeadingNode );
+		UE_API void AddUnversionedCookProperty( FLaunchProfileTreeNode& HeadingNode );
 		UE_API void AddCookProperty( FLaunchProfileTreeNode& HeadingNode );
 		UE_API void AddMapsToCookProperty( FLaunchProfileTreeNode& HeadingNode );
 		UE_API void AddAdditionalCookerOptionsProperty( FLaunchProfileTreeNode& HeadingNode );
@@ -127,9 +128,6 @@ namespace ProjectLauncher
 		UE_API void SetCook( bool bCook );
 		UE_API bool GetCook( ILauncherProfilePtr InProfile = nullptr ) const;
 
-		UE_API void SetIncrementalCookMode( ELauncherProfileIncrementalCookMode::Type Mode );
-		UE_API ELauncherProfileIncrementalCookMode::Type GetIncrementalCookMode() const;
-
 		UE_API void SetMapsToCook(TArray<FString> MapsToCook);
 		UE_API TArray<FString> GetMapsToCook() const;
 		UE_API EMapOption GetMapOption() const;
diff --git a/Engine/Plugins/Editor/MobileLauncherProfileWizard/Source/MobileLauncherProfileWizard/Private/AndroidProfileWizard.cpp b/Engine/Plugins/Editor/MobileLauncherProfileWizard/Source/MobileLauncherProfileWizard/Private/AndroidProfileWizard.cpp
--- a/Engine/Plugins/Editor/MobileLauncherProfileWizard/Source/MobileLauncherProfileWizard/Private/AndroidProfileWizard.cpp
+++ b/Engine/Plugins/Editor/MobileLauncherProfileWizard/Source/MobileLauncherProfileWizard/Private/AndroidProfileWizard.cpp
@@ -49,7 +49,7 @@ static void SetupAndroidAppProfile(ILauncherProfileRef& AppProfile, const FProfi
 	// Release settings
 	AppProfile->SetCreateReleaseVersion(true);
 	AppProfile->SetCreateReleaseVersionName(AndroidProfileConstants::AppReleaseName);
-	AppProfile->SetIncrementalCookMode(ELauncherProfileIncrementalCookMode::None);
+	AppProfile->SetIncrementalCooking(false);
 	AppProfile->SetCompressed(false);
 	AppProfile->SetDeployWithUnrealPak(true);
 
@@ -107,7 +107,7 @@ static void SetupAndroidDLCProfile(ILauncherProfileRef& DLCProfile, const FProfi
 	DLCProfile->SetDLCName(AndroidProfileConstants::DLCName);
 	DLCProfile->SetDLCIncludeEngineContent(true);
 		
-	DLCProfile->SetIncrementalCookMode(ELauncherProfileIncrementalCookMode::None);
+	DLCProfile->SetIncrementalCooking(false);
 	DLCProfile->SetCompressed(false);
 	DLCProfile->SetDeployWithUnrealPak(true);
 
diff --git a/Engine/Plugins/Editor/MobileLauncherProfileWizard/Source/MobileLauncherProfileWizard/Private/IOSProfileWizard.cpp b/Engine/Plugins/Editor/MobileLauncherProfileWizard/Source/MobileLauncherProfileWizard/Private/IOSProfileWizard.cpp
--- a/Engine/Plugins/Editor/MobileLauncherProfileWizard/Source/MobileLauncherProfileWizard/Private/IOSProfileWizard.cpp
+++ b/Engine/Plugins/Editor/MobileLauncherProfileWizard/Source/MobileLauncherProfileWizard/Private/IOSProfileWizard.cpp
@@ -49,7 +49,7 @@ static void SetupIOSAppProfile(ILauncherProfileRef& AppProfile, const FProfilePa
 	// Release settings
 	AppProfile->SetCreateReleaseVersion(true);
 	AppProfile->SetCreateReleaseVersionName(IOSProfileConstants::AppReleaseName);
-	AppProfile->SetIncrementalCookMode(ELauncherProfileIncrementalCookMode::None);
+	AppProfile->SetIncrementalCooking(false);
 	AppProfile->SetCompressed(false);
 	AppProfile->SetDeployWithUnrealPak(true);
 
@@ -104,7 +104,7 @@ static void SetupIOSDLCProfile(ILauncherProfileRef& DLCProfile, const FProfilePa
 	DLCProfile->SetDLCName(IOSProfileConstants::DLCName);
 	DLCProfile->SetDLCIncludeEngineContent(true);
 		
-	DLCProfile->SetIncrementalCookMode(ELauncherProfileIncrementalCookMode::None);
+	DLCProfile->SetIncrementalCooking(false);
 	DLCProfile->SetCompressed(false);
 	DLCProfile->SetDeployWithUnrealPak(true);
 
diff --git a/Engine/Source/Developer/LauncherServices/Private/Launcher/LauncherWorker.cpp b/Engine/Source/Developer/LauncherServices/Private/Launcher/LauncherWorker.cpp
--- a/Engine/Source/Developer/LauncherServices/Private/Launcher/LauncherWorker.cpp
+++ b/Engine/Source/Developer/LauncherServices/Private/Launcher/LauncherWorker.cpp
@@ -975,19 +975,9 @@ FString FLauncherWorker::CreateUATCommand( const ILauncherProfileRef& InProfile,
 		UATCommand += TEXT(" -distribution");
 	}
 
-	switch (InProfile->GetIncrementalCookMode())
+	if (InProfile->IsCookingIncrementally())
 	{
-		case ELauncherProfileIncrementalCookMode::ModifiedOnly:
-		{
-			UATCommand += TEXT(" -iterativecooking");
-		}
-		break;
-
-		case ELauncherProfileIncrementalCookMode::ModifiedAndDependencies:
-		{
-			UATCommand += TEXT(" -cookincremental");
-		}
-		break;
+		UATCommand += TEXT(" -iterativecooking");
 	}
 
 	if ( InProfile->IsIterateSharedCookedBuild() )
diff --git a/Engine/Source/Developer/LauncherServices/Private/Profiles/LauncherProfile.cpp b/Engine/Source/Developer/LauncherServices/Private/Profiles/LauncherProfile.cpp
--- a/Engine/Source/Developer/LauncherServices/Private/Profiles/LauncherProfile.cpp
+++ b/Engine/Source/Developer/LauncherServices/Private/Profiles/LauncherProfile.cpp
@@ -36,7 +36,7 @@ FString LexToStringLocalized(ELauncherProfileValidationErrors::Type Value)
 		case ELauncherProfileValidationErrors::NoPlatformSDKInstalled:
 			return LOCTEXT("NoPlatformSDKInstalled", "A required platform SDK is missing.").ToString();
 		case ELauncherProfileValidationErrors::UnversionedAndIncremental:
-			return LOCTEXT("UnversionedAndIncremental", "Unversioned build cannot be incremental.").ToString();
+			return LOCTEXT("UnversionedAndIncremental", "Unversioned build cannot be incremental when using LegacyIterative cook setting.").ToString();
 		case ELauncherProfileValidationErrors::GeneratingPatchesCanOnlyRunFromByTheBookCookMode:
 			return LOCTEXT("GeneratingPatchesCanOnlyRunFromByTheBookCookMode", "Generating patch requires cooking or zen snapshot import.").ToString();
 		case ELauncherProfileValidationErrors::GeneratingMultiLevelPatchesRequiresGeneratePatch:
diff --git a/Engine/Source/Developer/LauncherServices/Private/Profiles/LauncherProfile.h b/Engine/Source/Developer/LauncherServices/Private/Profiles/LauncherProfile.h
--- a/Engine/Source/Developer/LauncherServices/Private/Profiles/LauncherProfile.h
+++ b/Engine/Source/Developer/LauncherServices/Private/Profiles/LauncherProfile.h
@@ -4,6 +4,7 @@
 
 #include "CoreMinimal.h"
 #include "Misc/Guid.h"
+#include "Misc/ConfigCacheIni.h"
 #include "Modules/ModuleManager.h"
 #include "ILauncherProfile.h"
 #include "ILauncherServicesModule.h"
@@ -64,6 +65,7 @@ enum ELauncherVersion
 	LAUNCHERSERVICES_INCREMENTALCOOKMODE = 40,
 	LAUNCHERSERVICES_AUTOMATEDTESTS = 41,
 	LAUNCHERSERVICES_FASTITERATE = 42,
+	LAUNCHERSERVICES_REMOVEDINCREMENTALCOOKMODE = 43,
 	//ADD NEW STUFF HERE
 
 
@@ -881,14 +883,15 @@ public:
 	{
 		if ( CookMode != ELauncherProfileCookModes::DoNotCook )
 		{
-			return (IncrementalCookMode == ELauncherProfileIncrementalCookMode::ModifiedOnly);
+			return CookIncremental;
 		}
 		return false;
 	}
 
 	virtual ELauncherProfileIncrementalCookMode::Type GetIncrementalCookMode() const override
 	{
-		return IncrementalCookMode;
+		// fixme: worth reading LegacyIterative to honor this value?!
+		return CookIncremental ? ELauncherProfileIncrementalCookMode::ModifiedAndDependencies : ELauncherProfileIncrementalCookMode::None;
 	}
 
 
@@ -1040,7 +1043,6 @@ public:
 
 		// IMPORTANT: make sure to bump LAUNCHERSERVICES_PROFILEVERSION when modifying this!
 		bool BuildGame = false;
-		bool CookIncremental = false;
 		Archive << Id
 				<< Name
 				<< Description
@@ -1211,14 +1213,11 @@ public:
 			Archive << CustomBoolProperties;
 		}
 
-		if (Version >= LAUNCHERSERVICES_INCREMENTALCOOKMODE)
+		if (Version >= LAUNCHERSERVICES_INCREMENTALCOOKMODE && Version < LAUNCHERSERVICES_REMOVEDINCREMENTALCOOKMODE)
 		{
+			TEnumAsByte<ELauncherProfileIncrementalCookMode::Type> IncrementalCookMode = (CookIncremental ? ELauncherProfileIncrementalCookMode::ModifiedOnly : ELauncherProfileIncrementalCookMode::None);
 			Archive << IncrementalCookMode;
 		}
-		else if (Archive.IsLoading())
-		{
-			IncrementalCookMode = CookIncremental ? ELauncherProfileIncrementalCookMode::ModifiedOnly : ELauncherProfileIncrementalCookMode::None;
-		}
 
 		if (Version >= LAUNCHERSERVICES_AUTOMATEDTESTS)
 		{
@@ -1311,6 +1310,7 @@ public:
 		Writer.WriteValue("ProjectSpecified", ProjectSpecified);
 		Writer.WriteValue("ShareableProjectPath", ShareableProjectPath);
 		Writer.WriteValue("CookConfiguration", (int32)CookConfiguration);
+		Writer.WriteValue("CookIncremental", CookIncremental);
 		Writer.WriteValue("CookOptions", CookOptions);
 		Writer.WriteValue("CookMode", CookMode);
 		Writer.WriteValue("CookUnversioned", CookUnversioned);
@@ -1398,7 +1398,6 @@ public:
 		Writer.WriteValue("EditorArchitectures", EditorArchitectures);
 		Writer.WriteValue("CustomStringProperties", CustomStringProperties);
 		Writer.WriteValue("CustomBoolProperties", CustomBoolProperties);
-		Writer.WriteValue("IncrementalCookMode", IncrementalCookMode);
 		Writer.WriteValue("FastIterate", bFastIterate);
 
 		// serialize the default launch role
@@ -1962,8 +1961,6 @@ public:
 			return false;
 		}
 
-		bool CookIncremental;
-
 		FGuid::Parse(Object.GetStringField(TEXT("Id")), Id);
 		Name = Object.GetStringField(TEXT("Name"));
 		Description = Object.GetStringField(TEXT("Description"));
@@ -1971,7 +1968,6 @@ public:
 		ProjectSpecified = Object.GetBoolField(TEXT("ProjectSpecified"));
 		ShareableProjectPath = Object.GetStringField(TEXT("ShareableProjectPath"));
 		CookConfiguration = (EBuildConfiguration)((int32)Object.GetNumberField(TEXT("CookConfiguration")));
-		CookIncremental = Object.GetBoolField(TEXT("CookIncremental"));
 		CookOptions = Object.GetStringField(TEXT("CookOptions"));
 		CookMode = (TEnumAsByte<ELauncherProfileCookModes::Type>)((int32)Object.GetNumberField(TEXT("CookMode")));
 		CookUnversioned = Object.GetBoolField(TEXT("CookUnversioned"));
@@ -2170,13 +2166,14 @@ public:
 			}
 		}
 
-		if (Version >= LAUNCHERSERVICES_INCREMENTALCOOKMODE)
+		if (Version >= LAUNCHERSERVICES_INCREMENTALCOOKMODE && Version < LAUNCHERSERVICES_REMOVEDINCREMENTALCOOKMODE)
 		{
-			IncrementalCookMode = (TEnumAsByte<ELauncherProfileIncrementalCookMode::Type>)((int32)Object.GetNumberField(TEXT("IncrementalCookMode")));
+			ELauncherProfileIncrementalCookMode::Type IncrementalCookMode = (TEnumAsByte<ELauncherProfileIncrementalCookMode::Type>)((int32)Object.GetNumberField(TEXT("IncrementalCookMode")));
+			CookIncremental = (IncrementalCookMode != ELauncherProfileIncrementalCookMode::None);
 		}
-		else
+		else if (Version < LAUNCHERSERVICES_INCREMENTALCOOKMODE || Version >= LAUNCHERSERVICES_REMOVEDINCREMENTALCOOKMODE)
 		{
-			IncrementalCookMode = CookIncremental ? ELauncherProfileIncrementalCookMode::ModifiedOnly : ELauncherProfileIncrementalCookMode::None;
+			CookIncremental = Object.GetBoolField(TEXT("CookIncremental"));
 		}
 
 		if (Version >= LAUNCHERSERVICES_AUTOMATEDTESTS)
@@ -2282,7 +2279,7 @@ public:
 		CookConfiguration = FApp::GetBuildConfiguration();
 		CookMode = ELauncherProfileCookModes::OnTheFly;
 		CookOptions = FString();
-		IncrementalCookMode = ELauncherProfileIncrementalCookMode::None;
+		CookIncremental = false;
 		IterateSharedCookedBuild=false;
 		CookUnversioned = true;
 		Compressed = true;
@@ -2582,17 +2579,17 @@ public:
 
 	virtual void SetIncrementalCooking( bool Incremental ) override
 	{
-		SetIncrementalCookMode( Incremental ? ELauncherProfileIncrementalCookMode::ModifiedOnly : ELauncherProfileIncrementalCookMode::None );
+		if (Incremental != CookIncremental)
+		{
+			CookIncremental = Incremental;
+			Validate();
+		}
 	}
 
 	virtual void SetIncrementalCookMode( ELauncherProfileIncrementalCookMode::Type Mode ) override
 	{
-		if (IncrementalCookMode != Mode)
-		{
-			IncrementalCookMode = Mode;
-
-			Validate();
-		}
+		bool Incremental = (Mode != ELauncherProfileIncrementalCookMode::None);
+		SetIncrementalCooking(Incremental);
 	}
 
 
@@ -3173,7 +3170,10 @@ protected:
 			ValidationErrors.Add(ELauncherProfileValidationErrors::ZenPakStreamingRequiresDeployAndLaunch);
 		}
 
-		if (CookUnversioned && (IncrementalCookMode == ELauncherProfileIncrementalCookMode::ModifiedOnly) && ((CookMode == ELauncherProfileCookModes::ByTheBook) || (CookMode == ELauncherProfileCookModes::ByTheBookInEditor)))
+		bool bLegacyIterative = false;
+		GConfig->GetBool(TEXT("CookSettings"), TEXT("LegacyIterative"), bLegacyIterative, GEditorIni);
+
+		if (CookUnversioned && bLegacyIterative && CookIncremental && ((CookMode == ELauncherProfileCookModes::ByTheBook) || (CookMode == ELauncherProfileCookModes::ByTheBookInEditor)))
 		{
 			ValidationErrors.Add(ELauncherProfileValidationErrors::UnversionedAndIncremental);
 		}
@@ -3523,8 +3523,8 @@ private:
 	// is this build for distribution
 	bool ForDistribution;
 	
-	// Holds the incremental cook mode indicating whether only modified content should be cooked.
-	TEnumAsByte<ELauncherProfileIncrementalCookMode::Type> IncrementalCookMode;
+	// Holds a flag indicating whether only modified content should be cooked.
+	bool CookIncremental;
 
 	// hold a flag indicating if we want to iterate on the shared cooked build
 	bool IterateSharedCookedBuild;
diff --git a/Engine/Source/Developer/LauncherServices/Public/ILauncherProfile.h b/Engine/Source/Developer/LauncherServices/Public/ILauncherProfile.h
--- a/Engine/Source/Developer/LauncherServices/Public/ILauncherProfile.h
+++ b/Engine/Source/Developer/LauncherServices/Public/ILauncherProfile.h
@@ -902,7 +902,6 @@ public:
 	 * @return true if cooking incrementally, false otherwise.
 	 * @see SetIncrementalCooking
 	 */
-	UE_DEPRECATED(5.6, "Use GetIncrementalCookMode instead")
 	virtual bool IsCookingIncrementally( ) const = 0;
 
 	/**
@@ -911,6 +910,7 @@ public:
 	 * @return incremental cook mode
 	 * @see SetIncrementalCookMode
 	 */ 
+	UE_DEPRECATED(5.8, "Use IsCookingIncrementally again instead")
 	virtual ELauncherProfileIncrementalCookMode::Type GetIncrementalCookMode() const = 0;
 
 
@@ -1328,7 +1328,6 @@ public:
 	 * @param Incremental Whether cooking should be incremental.
 	 * @see IsCookingIncrementally
 	 */
-	UE_DEPRECATED(5.6, "Use SetIncrementalCookMode instead")
 	virtual void SetIncrementalCooking( bool Incremental ) = 0;
 
 	/**
@@ -1337,6 +1336,7 @@ public:
 	 * @param Mode the incremental cooking mode to use
 	 * @see GetIncrementalCookMode
 	 */
+	UE_DEPRECATED(5.8, "Use SetIncrementalCooking again instead")
 	virtual void SetIncrementalCookMode( ELauncherProfileIncrementalCookMode::Type Mode ) = 0;
 
 
diff --git a/Engine/Source/Developer/LegacyProjectLauncher/Private/Widgets/Cook/SProjectLauncherCookByTheBookSettings.cpp b/Engine/Source/Developer/LegacyProjectLauncher/Private/Widgets/Cook/SProjectLauncherCookByTheBookSettings.cpp
--- a/Engine/Source/Developer/LegacyProjectLauncher/Private/Widgets/Cook/SProjectLauncherCookByTheBookSettings.cpp
+++ b/Engine/Source/Developer/LegacyProjectLauncher/Private/Widgets/Cook/SProjectLauncherCookByTheBookSettings.cpp
@@ -1305,8 +1305,7 @@ void SProjectLauncherCookByTheBookSettings::HandleIncrementalCheckBoxCheckStateC
 
 	if (SelectedProfile.IsValid())
 	{
-		bool bIncremental = (NewState == ECheckBoxState::Checked);
-		SelectedProfile->SetIncrementalCookMode( bIncremental ? ELauncherProfileIncrementalCookMode::ModifiedOnly : ELauncherProfileIncrementalCookMode::None);
+		SelectedProfile->SetIncrementalCooking(NewState == ECheckBoxState::Checked);
 	}
 }
 
@@ -1317,7 +1316,7 @@ ECheckBoxState SProjectLauncherCookByTheBookSettings::HandleIncrementalCheckBoxI
 
 	if (SelectedProfile.IsValid())
 	{
-		if (SelectedProfile->GetIncrementalCookMode() != ELauncherProfileIncrementalCookMode::None)
+		if (SelectedProfile->IsCookingIncrementally())
 		{
 			return ECheckBoxState::Checked;
 		}
diff --git a/Engine/Source/Developer/LegacyProjectLauncher/Private/Widgets/Cook/SProjectLauncherCookOnTheFlySettings.cpp b/Engine/Source/Developer/LegacyProjectLauncher/Private/Widgets/Cook/SProjectLauncherCookOnTheFlySettings.cpp
--- a/Engine/Source/Developer/LegacyProjectLauncher/Private/Widgets/Cook/SProjectLauncherCookOnTheFlySettings.cpp
+++ b/Engine/Source/Developer/LegacyProjectLauncher/Private/Widgets/Cook/SProjectLauncherCookOnTheFlySettings.cpp
@@ -104,8 +104,7 @@ void SProjectLauncherCookOnTheFlySettings::HandleIncrementalCheckBoxCheckStateCh
 
 	if (SelectedProfile.IsValid())
 	{
-		bool bIncremental = (NewState == ECheckBoxState::Checked);
-		SelectedProfile->SetIncrementalCookMode( bIncremental ? ELauncherProfileIncrementalCookMode::ModifiedOnly : ELauncherProfileIncrementalCookMode::None);
+		SelectedProfile->SetIncrementalCooking(NewState == ECheckBoxState::Checked);
 	}
 }
 
@@ -116,7 +115,7 @@ ECheckBoxState SProjectLauncherCookOnTheFlySettings::HandleIncrementalCheckBoxIs
 
 	if (SelectedProfile.IsValid())
 	{
-		if (SelectedProfile->GetIncrementalCookMode() != ELauncherProfileIncrementalCookMode::None)
+		if (SelectedProfile->IsCookingIncrementally())
 		{
 			return ECheckBoxState::Checked;
 		}
diff --git a/Engine/Source/Editor/UnrealEd/Private/PlayLevelViaLauncher.cpp b/Engine/Source/Editor/UnrealEd/Private/PlayLevelViaLauncher.cpp
--- a/Engine/Source/Editor/UnrealEd/Private/PlayLevelViaLauncher.cpp
+++ b/Engine/Source/Editor/UnrealEd/Private/PlayLevelViaLauncher.cpp
@@ -294,7 +294,7 @@ void UEditorEngine::StartPlayUsingLauncherSession(FRequestPlaySessionParams& InR
 	}
 	LauncherProfile->SetCookMode(CurrentLauncherCookMode);
 	LauncherProfile->SetUnversionedCooking(!bIncrementalCooking); // Unversioned cooking is not allowed with incremental cooking
-	LauncherProfile->SetIncrementalCookMode(bIncrementalCooking ? ELauncherProfileIncrementalCookMode::ModifiedOnly : ELauncherProfileIncrementalCookMode::None);
+	LauncherProfile->SetIncrementalCooking(bIncrementalCooking);
 	SetCookOption(TEXTVIEW("-IgnoreIniSettingsOutOfDate"), bIncrementalCooking && CookerSettings.bIgnoreIniSettingsOutOfDateForIteration);
 	SetCookOption(TEXTVIEW("-IgnoreScriptPackagesOutOfDate"), bIncrementalCooking && CookerSettings.bIgnoreScriptPackagesOutOfDateForIteration);
 	SetCookOption(TEXTVIEW("-IterateSharedCookedbuild"), bIncrementalCooking && ExperimentalSettings.bSharedCookedBuilds);
diff --git a/Engine/Source/Programs/UnrealFrontend/Private/Commands/UserInterfaceCommand.cpp b/Engine/Source/Programs/UnrealFrontend/Private/Commands/UserInterfaceCommand.cpp
--- a/Engine/Source/Programs/UnrealFrontend/Private/Commands/UserInterfaceCommand.cpp
+++ b/Engine/Source/Programs/UnrealFrontend/Private/Commands/UserInterfaceCommand.cpp
@@ -10,6 +10,7 @@
 #include "Containers/Ticker.h"
 #include "Misc/CommandLine.h"
 #include "Misc/ConfigCacheIni.h"
+#include "Misc/ConfigContext.h"
 #include "Misc/App.h"
 #include "Misc/CoreDelegates.h"
 #include "Framework/Docking/TabManager.h"
@@ -42,8 +43,15 @@ void FUserInterfaceCommand::Run(  )
 	// ensure target platform manager is referenced early as it must be created on the main thread
 	FModuleManager::Get().LoadModuleChecked("DesktopPlatform");
 	FConfigCacheIni::InitializeConfigSystem();
+
 	GetTargetPlatformManager();
 
+#if !WITH_EDITOR
+	// some items in Launcher Services need the editor ini settings - make sure they're loaded if they've not been done already
+	FConfigContext Context = FConfigContext::ReadIntoGConfig(); 
+	Context.Load(TEXT("Editor"), GEditorIni);
+#endif
+
 	// Crank up a normal Slate application using the platform's standalone renderer.
 	FSlateApplication::InitializeAsStandaloneApplication(GetStandardStandaloneRenderer());
 
