Fix for Project Launcher 2 failing to apply the Additional Command Line when automated tests are enabled.

diff --git a/Engine/Source/Developer/LauncherServices/Private/Launcher/LauncherWorker.cpp b/Engine/Source/Developer/LauncherServices/Private/Launcher/LauncherWorker.cpp
--- a/Engine/Source/Developer/LauncherServices/Private/Launcher/LauncherWorker.cpp
+++ b/Engine/Source/Developer/LauncherServices/Private/Launcher/LauncherWorker.cpp
@@ -1240,10 +1240,27 @@ FString FLauncherWorker::CreateUATCommandForAutomatedTest( const ILauncherProfil
 		}
 	}
 
-	// final extra command line parameters
-	FString AdditionalCommandLine = InAutomatedTest->GetAdditionalCommandLine();
-	LauncherProfileManager->OnPostProcessAutomatedTestCommandLine().Broadcast(InAutomatedTest, InProfile, AdditionalCommandLine);
+	// Profile 'additional command line parameters' are commands for the app wrapped in the flag -addcmdline in CreateUATCommand. Wrap the value if 
+	// it exists in the automated test equivalent which is -Args. 
+	FString AdditionalCommandLine;
+
+	// Append the profile command line parameters
+	FString ProfileAdditionalCommandLine = InProfile->GetAdditionalCommandLineParameters();
+	if (!ProfileAdditionalCommandLine.IsEmpty())
+	{
+		AdditionalCommandLine = FString::Printf(TEXT("-Args=\"%s\""), *ProfileAdditionalCommandLine.Replace(TEXT("\""), TEXT("\\\"")));
+	}
 
+	// Append the tests additional command line parameters
+	FString TestAdditionalCommandLine = InAutomatedTest->GetAdditionalCommandLine();
+	if (!TestAdditionalCommandLine.IsEmpty())
+	{
+		AdditionalCommandLine += TEXT(" ");
+		AdditionalCommandLine += TestAdditionalCommandLine;
+	}
+
+	// Dispatch the event which may further customize the command line passed to the automated test.
+	LauncherProfileManager->OnPostProcessAutomatedTestCommandLine().Broadcast(InAutomatedTest, InProfile, AdditionalCommandLine);
 	if (!AdditionalCommandLine.IsEmpty())
 	{
 		AutomatedTestCommand += TEXT(" ");
