fix for Project Launcher not generating pak files when only importing a snapshot & not cooking

#rnx

diff --git a/Engine/Source/Developer/LauncherServices/Private/Launcher/LauncherWorker.cpp b/Engine/Source/Developer/LauncherServices/Private/Launcher/LauncherWorker.cpp
--- a/Engine/Source/Developer/LauncherServices/Private/Launcher/LauncherWorker.cpp
+++ b/Engine/Source/Developer/LauncherServices/Private/Launcher/LauncherWorker.cpp
@@ -967,6 +967,18 @@ FString FLauncherWorker::CreateUATCommand( const ILauncherProfileRef& InProfile,
 		{
 			UATCommand += FString::Printf(TEXT(" -skippak -ZenWorkspaceSharePath=\\\"%s\\\" "), *InProfile->GetZenPakStreamingPath() );
 		}
+		else if (InProfile->IsImportingZenSnapshot() && InProfile->IsPackingWithUnrealPak())
+		{
+			UATCommand += TEXT(" -pak");
+			if (InProfile->IsUsingIoStore())
+			{
+				UATCommand += TEXT(" -iostore");
+			}
+			if (InProfile->IsCompressed())
+			{
+				UATCommand += TEXT(" -compressed");
+			}
+		}
 		break;
 	}
 
