don't delete launcher profiles that can't be loaded - could be working with an older editor etc.

#rnx
#Synced-CL 48448085

diff --git a/Engine/Source/Developer/LauncherServices/Private/Profiles/LauncherProfileManager.cpp b/Engine/Source/Developer/LauncherServices/Private/Profiles/LauncherProfileManager.cpp
--- a/Engine/Source/Developer/LauncherServices/Private/Profiles/LauncherProfileManager.cpp
+++ b/Engine/Source/Developer/LauncherServices/Private/Profiles/LauncherProfileManager.cpp
@@ -604,7 +604,7 @@ void FLauncherProfileManager::LoadProfiles( )
 			}
 			else
 			{
-				IFileManager::Get().Delete(*ProfileFilePath);
+				UE_LOG(LogLauncherProfile, Warning, TEXT("could not load %s - ignoring"), *ProfileFilePath);
 			}
 		}
 	}
