Launcher profiles extensions can define generic UAT commands that can be run before & after the main BuildCookRun operation

other changes:
- automated tests must specify a positive Order rather than an Priority which is used when sorting against all UAT commands

#rnx
#Synced-CL 48803560

diff --git a/Engine/Plugins/Developer/ProjectLauncher/Source/CommonLaunchExtensions/Private/BootTest/BootTestLaunchExtension.cpp b/Engine/Plugins/Developer/ProjectLauncher/Source/CommonLaunchExtensions/Private/BootTest/BootTestLaunchExtension.cpp
--- a/Engine/Plugins/Developer/ProjectLauncher/Source/CommonLaunchExtensions/Private/BootTest/BootTestLaunchExtension.cpp
+++ b/Engine/Plugins/Developer/ProjectLauncher/Source/CommonLaunchExtensions/Private/BootTest/BootTestLaunchExtension.cpp
@@ -20,7 +20,7 @@ void FBootTestLaunchExtensionInstance::CustomizeTree( ProjectLauncher::FLaunchPr
 	AddDefaultHeading(ProfileTreeData)
 		.AddBoolean(LOCTEXT("WindowedLabel","Windowed"),
 		{
-			.GetValue = [this]()          { return GetConfigBool(EConfig::PerProfile, TEXT("Windowed"));},
+			.GetValue = [this]()          { return GetConfigBool(EConfig::PerProfile, TEXT("Windowed")); },
 			.SetValue = [this](bool bVal) { SetConfigBool(EConfig::PerProfile, TEXT("Windowed"), bVal ); },
 			.IsVisible = IsVisible,
 		}
@@ -44,10 +44,13 @@ const FString FBootTestLaunchExtensionInstance::GetTestInternalName() const
 void FBootTestLaunchExtensionInstance::OnTestAdded( ILauncherProfileAutomatedTestRef AutomatedTest )
 {
 	AutomatedTest->SetTests(TEXT("UE.BootTest"));
-	AutomatedTest->SetPriority(1000); // Boot Test is likely the most lightweight test there is, so we want to run it first - give it a high priority
 }
 
-
+int32 FBootTestLaunchExtensionInstance::GetTestOrder() const
+{
+	// Boot Test is likely the most lightweight test there is, so we want to run it first - give it a low order
+	return 1; 
+}
 
 
 TSharedPtr<ProjectLauncher::FLaunchExtensionInstance> FBootTestLaunchExtension::CreateInstanceForProfile( ProjectLauncher::FLaunchExtensionInstance::FArgs& InArgs )
diff --git a/Engine/Plugins/Developer/ProjectLauncher/Source/CommonLaunchExtensions/Private/BootTest/BootTestLaunchExtension.h b/Engine/Plugins/Developer/ProjectLauncher/Source/CommonLaunchExtensions/Private/BootTest/BootTestLaunchExtension.h
--- a/Engine/Plugins/Developer/ProjectLauncher/Source/CommonLaunchExtensions/Private/BootTest/BootTestLaunchExtension.h
+++ b/Engine/Plugins/Developer/ProjectLauncher/Source/CommonLaunchExtensions/Private/BootTest/BootTestLaunchExtension.h
@@ -14,6 +14,7 @@ public:
 
 	virtual void OnTestAdded( ILauncherProfileAutomatedTestRef AutomatedTest ) override;
 	virtual const FString GetTestInternalName() const override;
+	virtual int32 GetTestOrder() const override;
 
 	virtual void CustomizeTree( ProjectLauncher::FLaunchProfileTreeData& ProfileTreeData ) override;
 	virtual void CustomizeAutomatedTestCommandLine( FString& InOutCommandLine ) override;
diff --git a/Engine/Plugins/Developer/ProjectLauncher/Source/ProjectLauncher/Private/Extension/LaunchExtension.cpp b/Engine/Plugins/Developer/ProjectLauncher/Source/ProjectLauncher/Private/Extension/LaunchExtension.cpp
--- a/Engine/Plugins/Developer/ProjectLauncher/Source/ProjectLauncher/Private/Extension/LaunchExtension.cpp
+++ b/Engine/Plugins/Developer/ProjectLauncher/Source/ProjectLauncher/Private/Extension/LaunchExtension.cpp
@@ -101,15 +101,16 @@ namespace ProjectLauncher
 		}
 	}
 
-	void ApplyExtensionVariablesForAutomatedTest( const ILauncherProfileAutomatedTestRef& InAutomatedTest, const ILauncherProfileRef& InProfile, FString& InOutCommandLine, TSharedRef<FModel> InModel )
+
+	void ApplyExtensionVariablesForUATCommand( FUATCommandPostProcessParameters& InPostProcessParameters, TSharedRef<FModel> InModel )
 	{
 		// take a snapshot of the extension instances for this profile
 		TArray<TSharedPtr<FLaunchExtensionInstance>> ExtensionInstances;
 		{
 			FScopeLock Lock(&GExtensionsCS);
-			if (GExtensionInstances.Contains(InProfile))
+			if (GExtensionInstances.Contains(InPostProcessParameters.Profile))
 			{
-				ExtensionInstances = GExtensionInstances[InProfile];
+				ExtensionInstances = GExtensionInstances[InPostProcessParameters.Profile];
 			}
 		}
 
@@ -124,19 +125,31 @@ namespace ProjectLauncher
 				{
 					for (const FString& Variable : Variables)
 					{
-						if (InOutCommandLine.Contains(Variable, ESearchCase::IgnoreCase))
+						if (InPostProcessParameters.UATCommandLine.Contains(Variable, ESearchCase::IgnoreCase))
 						{
 							FString VariableValue;
 							if (ExtensionInstance->GetExtensionVariableValue( Variable, VariableValue ))
 							{
-								InOutCommandLine.ReplaceInline( *Variable, *VariableValue, ESearchCase::IgnoreCase );
+								InPostProcessParameters.UATCommandLine.ReplaceInline( *Variable, *VariableValue, ESearchCase::IgnoreCase );
 							}
 						}
 					}
 				}
 
 				// allow for advanced command line customization
-				ExtensionInstance->CustomizeAutomatedTestCommandLine( InAutomatedTest, InOutCommandLine );
+				bool bKnownCommand = false;
+				const ILauncherProfileAutomatedTestPtr AutomatedTest = InPostProcessParameters.UATCommand->AsAutomatedTest();
+				if (AutomatedTest.IsValid())
+				{
+					bKnownCommand = ExtensionInstance->CustomizeAutomatedTestCommandLine( AutomatedTest.ToSharedRef(), InPostProcessParameters.UATCommandLine );
+				}
+				else
+				{
+					bKnownCommand = ExtensionInstance->CustomizeUATCommandLine( InPostProcessParameters.UATCommand, InPostProcessParameters.UATCommandLine );
+				}
+
+				// don't skip this command if something knows about it
+				InPostProcessParameters.bShouldBeSkipped &= !bKnownCommand;
 			}
 		}
 	}
@@ -250,6 +263,15 @@ namespace ProjectLauncher
 	};
 
 
+	const FLaunchExtensionInstance::FExtensionsMenuEntry FLaunchExtensionInstance::FExtensionsMenuEntry::UATCommands =
+	{
+		.SectionName = DefaultSectionName,
+		.SubmenuName = TEXT("UATCommands"),
+		.SectionDisplayName = DefaultSectionDisplayName,
+		.SubmenuDisplayName = LOCTEXT("UATCommandsSubMenu", "More Buid Steps"),
+		.Type = FLaunchExtensionInstance::FExtensionsMenuEntry::Type_SubMenu,
+	};
+
 
 
 
@@ -721,7 +743,7 @@ namespace ProjectLauncher
 
 	FLaunchProfileTreeNode& FLaunchExtensionInstance::AddDefaultHeading(ProjectLauncher::FLaunchProfileTreeData& ProfileTreeData, FText OverrideDisplayName) const
 	{
-		return ProfileTreeData.AddHeading(Extension->GetInternalName(), OverrideDisplayName);
+		return ProfileTreeData.AddHeading(Extension->GetInternalName(), OverrideDisplayName, GetHeadingSortOrder() );
 	}
 
 	void FLaunchExtensionInstance::GetExtensionsMenuEntry( FExtensionsMenuEntry& MenuEntry ) const
@@ -732,6 +754,85 @@ namespace ProjectLauncher
 
 
 
+	FUATCommandLaunchExtensionInstance::FUATCommandLaunchExtensionInstance( FArgs& InArgs )
+		: FLaunchExtensionInstance(InArgs)
+	{
+	}
+
+	void FUATCommandLaunchExtensionInstance::InternalInitialize()
+	{
+		ILauncherProfileUATCommandPtr UATCommand = GetUATCommand();
+		if (UATCommand.IsValid())
+		{
+			OnUATCommandAdded(UATCommand.ToSharedRef());
+			UATCommand->SetOrder(GetUATCommandOrder());
+		}
+	}
+
+	void FUATCommandLaunchExtensionInstance::MakeCustomExtensionSubmenu( FMenuBuilder& MenuBuilder )
+	{
+		auto ToggleUATCommandEnabled = [this]()
+		{
+			ILauncherProfileUATCommandPtr ExistingUATCommand = GetUATCommand();
+			if (ExistingUATCommand.IsValid())
+			{
+				OnUATCommandRemoved(ExistingUATCommand.ToSharedRef());
+				GetProfile()->RemoveUATCommand(*GetUATCommandInternalName());
+			}
+			else
+			{
+				ILauncherProfileUATCommandRef UATCommand = GetProfile()->FindOrAddUATCommand(*GetUATCommandInternalName());
+				OnUATCommandAdded(UATCommand);
+				UATCommand->SetOrder(GetUATCommandOrder());
+			}
+
+			BroadcastPropertyChanged();
+		};
+
+		MenuBuilder.AddMenuEntry(
+			GetExtension()->GetDisplayName(),
+			FText::GetEmpty(),
+			FSlateIcon(),
+			FUIAction(
+				FExecuteAction::CreateLambda(ToggleUATCommandEnabled),
+				FCanExecuteAction(),
+				FIsActionChecked::CreateRaw(this, &FUATCommandLaunchExtensionInstance::IsUATCommandActive )
+			),
+			NAME_None,
+			EUserInterfaceActionType::Check);
+	}
+
+	bool FUATCommandLaunchExtensionInstance::IsUATCommandActive() const
+	{
+		return GetUATCommand() != nullptr;
+	}
+
+	ILauncherProfileUATCommandPtr FUATCommandLaunchExtensionInstance::GetUATCommand() const
+	{
+		return GetProfile()->GetUATCommand(*GetUATCommandInternalName());
+	}
+
+	void FUATCommandLaunchExtensionInstance::GetExtensionsMenuEntry( FExtensionsMenuEntry& MenuEntry ) const
+	{
+		MenuEntry = FExtensionsMenuEntry::UATCommands;
+	}
+
+
+	bool FUATCommandLaunchExtensionInstance::CustomizeUATCommandLine( const ILauncherProfileUATCommandRef& InUATCommand, FString& InOutCommandLine )
+	{
+		if (InUATCommand == GetUATCommand())
+		{
+			CustomizeUATCommandLine(InOutCommandLine);
+			return true;
+		}
+
+		return false;
+	}
+
+
+
+
+
 
 
 
@@ -747,6 +848,10 @@ namespace ProjectLauncher
 		if (AutomatedTest.IsValid())
 		{
 			OnTestAdded(AutomatedTest.ToSharedRef());
+
+			int32 Order = GetTestOrder();
+			check(Order > 0);
+			AutomatedTest->SetOrder(Order);
 		}
 	}
 
@@ -764,6 +869,10 @@ namespace ProjectLauncher
 			{
 				ILauncherProfileAutomatedTestRef AutomatedTest = GetProfile()->FindOrAddAutomatedTest(*GetTestInternalName());
 				OnTestAdded(AutomatedTest);
+
+				int32 Order = GetTestOrder();
+				check(Order > 0);
+				AutomatedTest->SetOrder(Order);
 			}
 
 			BroadcastPropertyChanged();
@@ -798,12 +907,15 @@ namespace ProjectLauncher
 	}
 
 
-	void FAutomatedTestLaunchExtensionInstance::CustomizeAutomatedTestCommandLine( const ILauncherProfileAutomatedTestRef& InAutomatedTest, FString& InOutCommandLine )
+	bool FAutomatedTestLaunchExtensionInstance::CustomizeAutomatedTestCommandLine( const ILauncherProfileAutomatedTestRef& InAutomatedTest, FString& InOutCommandLine )
 	{
 		if (InAutomatedTest == GetTest())
 		{
 			CustomizeAutomatedTestCommandLine(InOutCommandLine);
+			return true;
 		}
+
+		return false;
 	}
 }
 
diff --git a/Engine/Plugins/Developer/ProjectLauncher/Source/ProjectLauncher/Private/Model/ProjectLauncherModel.cpp b/Engine/Plugins/Developer/ProjectLauncher/Source/ProjectLauncher/Private/Model/ProjectLauncherModel.cpp
--- a/Engine/Plugins/Developer/ProjectLauncher/Source/ProjectLauncher/Private/Model/ProjectLauncherModel.cpp
+++ b/Engine/Plugins/Developer/ProjectLauncher/Source/ProjectLauncher/Private/Model/ProjectLauncherModel.cpp
@@ -36,7 +36,7 @@ namespace ProjectLauncher
 	{
 		// register callbacks
 		ProfileManager->OnPostProcessLaunchCommandLine().AddRaw( this, &FModel::OnModifyLaunchCommandLine );
-		ProfileManager->OnPostProcessAutomatedTestCommandLine().AddRaw( this, &FModel::OnModifyAutomatedTestCommandLine );
+		ProfileManager->OnPostProcessUATCommandLine().AddRaw( this, &FModel::OnModifyUATCommandLine );
 
 		ProfileManager->OnProfileAdded().AddRaw(this, &FModel::HandleProfileManagerProfileAdded);
 		ProfileManager->OnProfileRemoved().AddRaw(this, &FModel::HandleProfileManagerProfileRemoved);
@@ -89,7 +89,7 @@ namespace ProjectLauncher
 		}
 
 		ProfileManager->OnPostProcessLaunchCommandLine().RemoveAll(this);
-		ProfileManager->OnPostProcessAutomatedTestCommandLine().RemoveAll(this);
+		ProfileManager->OnPostProcessUATCommandLine().RemoveAll(this);
 
 		ProfileManager->OnProfileAdded().RemoveAll(this);
 		ProfileManager->OnProfileRemoved().RemoveAll(this);
@@ -760,12 +760,19 @@ PRAGMA_ENABLE_DEPRECATION_WARNINGS
 
 	void FModel::ReadProjectSettingsFromConfig( FConfigCacheIni& InConfig, const FString& InProjectPath, FProjectSettings& OutResult )
 	{
+		// read build settings
+		const TCHAR* BuildSettings = TEXT("/Script/BuildSettings.BuildSettings");
+		InConfig.GetString(BuildSettings, TEXT("DefaultEditorTarget"), OutResult.DefaultEditorTargetName, GEngineIni);
+		InConfig.GetString(BuildSettings, TEXT("DefaultServerTarget"), OutResult.DefaultServerTargetName, GEngineIni);
+		InConfig.GetString(BuildSettings, TEXT("DefaultClientTarget"), OutResult.DefaultClientTargetName, GEngineIni);
+
 		// read project packaging settings
 		const TCHAR* ProjectPackagingConfigSection = TEXT("/Script/UnrealEd.ProjectPackagingSettings");
 		InConfig.GetBool(ProjectPackagingConfigSection, TEXT("bUseZenStore"), OutResult.bUseZenStore, GGameIni);
 		InConfig.GetBool(ProjectPackagingConfigSection, TEXT("bUseIoStore"), OutResult.bUseIoStore, GGameIni);
 		InConfig.GetBool(ProjectPackagingConfigSection, TEXT("bEnablePakStreaming"), OutResult.bHasAutomaticZenPakStreamingWorkspaceCreation, GGameIni);
 		InConfig.GetString(ProjectPackagingConfigSection, TEXT("BuildTarget"), OutResult.DefaultBuildTargetName, GGameIni);
+
 		if (OutResult.DefaultBuildTargetName.IsEmpty())
 		{
 			const TArray<FTargetInfo>& BuildTargets = FDesktopPlatformModule::Get()->GetTargetsForProject(*InProjectPath);
@@ -781,6 +788,7 @@ PRAGMA_ENABLE_DEPRECATION_WARNINGS
 			// ... fixme ...
 		}
 
+
 		
 		// read zen settings
 #if UE_WITH_ZEN
@@ -917,14 +925,13 @@ PRAGMA_ENABLE_DEPRECATION_WARNINGS
 	}
 
 
-	extern void ApplyExtensionVariablesForAutomatedTest( const ILauncherProfileAutomatedTestRef& InAutomatedTest, const ILauncherProfileRef& InProfile, FString& InOutCommandLine, TSharedRef<FModel> InModel );
+	extern void ApplyExtensionVariablesForUATCommand( FUATCommandPostProcessParameters& PostProcessParameters, TSharedRef<FModel> Model );
 
-	void FModel::OnModifyAutomatedTestCommandLine( const ILauncherProfileAutomatedTestRef& InAutomatedTest, const ILauncherProfileRef& InProfile, FString& InOutCommandLine )
+	void FModel::OnModifyUATCommandLine( FUATCommandPostProcessParameters& PostProcessParameters )
 	{
-		ApplyExtensionVariablesForAutomatedTest(InAutomatedTest, InProfile, InOutCommandLine, AsShared() );
+		ApplyExtensionVariablesForUATCommand(PostProcessParameters, AsShared());
 	}
 
-
 	bool FModel::AreExtensionsEnabled() const
 	{
 #if WITH_EDITOR
diff --git a/Engine/Plugins/Developer/ProjectLauncher/Source/ProjectLauncher/Public/Extension/LaunchExtension.h b/Engine/Plugins/Developer/ProjectLauncher/Source/ProjectLauncher/Public/Extension/LaunchExtension.h
--- a/Engine/Plugins/Developer/ProjectLauncher/Source/ProjectLauncher/Public/Extension/LaunchExtension.h
+++ b/Engine/Plugins/Developer/ProjectLauncher/Source/ProjectLauncher/Public/Extension/LaunchExtension.h
@@ -86,7 +86,12 @@ namespace ProjectLauncher
 		/**
 		 * Advanced hook to allow any advanced modification of the command line when an automated test is launched
 		 */
-		virtual void CustomizeAutomatedTestCommandLine( const ILauncherProfileAutomatedTestRef& InAutomatedTest, FString& InOutCommandLine ) {};
+		virtual bool CustomizeAutomatedTestCommandLine( const ILauncherProfileAutomatedTestRef& InAutomatedTest, FString& InOutCommandLine ) { return false; }
+
+		/**
+		 * Advanced hook to allow any advanced modification of the command line when a UAT command test is launched
+		 */
+		virtual bool CustomizeUATCommandLine( const ILauncherProfileUATCommandRef& InUATCommand, FString& InOutCommandLine ) { return false; }
 
 		/**
 		 * Notification callback when a property has been changed in our profile
@@ -114,6 +119,7 @@ namespace ProjectLauncher
 			static UE_API const FExtensionsMenuEntry OwnSection;
 			static UE_API const FExtensionsMenuEntry Deprecated;
 			static UE_API const FExtensionsMenuEntry AutomatedTests;
+			static UE_API const FExtensionsMenuEntry UATCommands;
 
 			// internal names
 			FString SectionName; // if this is empty, the extension will be placed in its own section
@@ -371,6 +377,8 @@ namespace ProjectLauncher
 		UE_API bool IsParameterGroup( const FString& InParameter ) const;
 		UE_API bool TryRemoveParameterGroup( const FString& InParameter );
 
+		virtual int32 GetHeadingSortOrder() const { return 0; }
+
 
 		ILauncherProfileRef Profile;
 		ILauncherProfileRef DefaultProfile;
@@ -385,6 +393,77 @@ namespace ProjectLauncher
 	};
 
 
+	/**
+	 * Helper base class for an UAT command launch extension
+	 * Creates an item in the extension menu that will toggle the UAT command on and off
+	 * note: It is not necessary for a UAT command launch extension to use this class
+	 */
+	class FUATCommandLaunchExtensionInstance : public FLaunchExtensionInstance
+	{
+	public:
+		UE_API FUATCommandLaunchExtensionInstance( FArgs& InArgs );
+
+		/*
+		 * Populate the custom extension menu for this UAT command launch extension. This will add a menu item to toggle the UAT command on and off
+		 */
+		UE_API virtual void MakeCustomExtensionSubmenu( FMenuBuilder& MenuBuilder ) override;
+
+		/**
+		 * Describes how this extension is represented in the extensions menu - a common section for all UAT commands
+		 */
+		UE_API virtual void GetExtensionsMenuEntry( FExtensionsMenuEntry& MenuEntry ) const override;
+		
+		/**
+		 * Determine if this UAT Command is currently active
+		 */
+		UE_API bool IsUATCommandActive() const;
+
+	protected:
+		UE_API bool CustomizeUATCommandLine( const ILauncherProfileUATCommandRef& InUATCommand, FString& InOutCommandLine ) override;
+		virtual bool HasCustomExtensionMenu() const { return true; }
+
+		/**
+		 * Advanced hook to allow any advanced modification of the command line when this extension's UAT Command is launched
+		 */
+		virtual void CustomizeUATCommandLine( FString& InOutCommandLine ) {}
+
+		/**
+		 * Hook to allow any modification of the UAT Command & profile once it has been created. At least the UAT Command name should be specified
+		 */
+		virtual void OnUATCommandAdded( ILauncherProfileUATCommandRef UATCommand ) = 0;
+
+		/**
+		 * Advanced hook to allow any modification of the profile just before the UAT Command is removed
+		 */
+		virtual void OnUATCommandRemoved( ILauncherProfileUATCommandRef UATCommand ) {}
+
+		/**
+		 * Access this extension's UAT Command if it has been created
+		 */
+		UE_API ILauncherProfileUATCommandPtr GetUATCommand() const;
+
+		/**
+		 * Internal name for the UAT command. Must be globally unique, so recommended to include the extension name
+		 */
+		virtual const FString GetUATCommandInternalName() const = 0;
+
+		/**
+		 * Determine the order of this UAT command, negative runs before BuildCookRun, positive runs after
+		 */
+		UE_API virtual int32 GetUATCommandOrder() const { return INT32_MAX; }
+
+		/**
+		 * Internal intitialization function, called immediately after construction - don't call directly
+		 */
+		UE_API virtual void InternalInitialize() override;
+
+	private:
+		virtual int32 GetHeadingSortOrder() const override { return GetUATCommandOrder(); }
+	};
+
+
+
+
 	/**
 	 * Helper base class for an automated test launch extension
 	 * Creates an item in the extension menu that will toggle the test on and off
@@ -411,7 +490,7 @@ namespace ProjectLauncher
 		UE_API bool IsTestActive() const;
 
 	protected:
-		UE_API void CustomizeAutomatedTestCommandLine( const ILauncherProfileAutomatedTestRef& InAutomatedTest, FString& InOutCommandLine ) override;
+		UE_API bool CustomizeAutomatedTestCommandLine( const ILauncherProfileAutomatedTestRef& InAutomatedTest, FString& InOutCommandLine ) override;
 		virtual bool HasCustomExtensionMenu() const { return true; }
 
 		/**
@@ -439,10 +518,18 @@ namespace ProjectLauncher
 		 */
 		virtual const FString GetTestInternalName() const = 0;
 
+		/**
+		 * Determine the ordering for this test. lower number = run earlier
+		 */
+		virtual int32 GetTestOrder() const { return INT32_MAX; }
+
 		/**
 		 * Internal intitialization function, called immediately after construction - don't call directly
 		 */
 		UE_API virtual void InternalInitialize() override;
+
+	private:
+		virtual int32 GetHeadingSortOrder() const override { return GetTestOrder(); }
 	};
 
 
diff --git a/Engine/Plugins/Developer/ProjectLauncher/Source/ProjectLauncher/Public/Model/ProjectLauncherModel.h b/Engine/Plugins/Developer/ProjectLauncher/Source/ProjectLauncher/Public/Model/ProjectLauncherModel.h
--- a/Engine/Plugins/Developer/ProjectLauncher/Source/ProjectLauncher/Public/Model/ProjectLauncherModel.h
+++ b/Engine/Plugins/Developer/ProjectLauncher/Source/ProjectLauncher/Public/Model/ProjectLauncherModel.h
@@ -6,6 +6,7 @@
 #include "Containers/Array.h"
 #include "Delegates/Delegate.h"
 #include "ILauncherProfileManager.h"
+#include "ILauncherWorker.h"
 #include "Templates/SharedPointer.h"
 #include "Misc/DataDrivenPlatformInfoRegistry.h"
 
@@ -97,8 +98,11 @@ namespace ProjectLauncher
 		// whether Zen Server is allowed to accept connections from remote machines (i.e. console devkits, phones etc)
 		bool bAllowRemoteNetworkService = false;
 
-		// default build target name
+		// default build target names
 		FString DefaultBuildTargetName;
+		FString DefaultServerTargetName;
+		FString DefaultClientTargetName;
+		FString DefaultEditorTargetName;
 
 		// cached project config ini
 		FConfigCacheIni* Config = nullptr;
@@ -284,7 +288,7 @@ namespace ProjectLauncher
 
 		/* Modify launch command lines to inject extension variables */
 		UE_API void OnModifyLaunchCommandLine( const ILauncherProfileRef& InProfile, EBuildTargetType InBuildTargetType, FString& InOutCommandLine );
-		UE_API void OnModifyAutomatedTestCommandLine(const ILauncherProfileAutomatedTestRef& InAutomatedTest, const ILauncherProfileRef& InProfile, FString& InOutCommandLine);
+		UE_API void OnModifyUATCommandLine( FUATCommandPostProcessParameters& PostProcessParameters );
 
 		/* Get map paths for a project, caching results */
 		UE_API TArray<FString> GetAndCacheMapPaths(const FString& InOptionalProjectPath, bool bIncludeNonContentDirMaps = false);
diff --git a/Engine/Plugins/Developer/ProjectLauncher/Source/ProjectLauncher/Public/ProfileTree/LaunchProfileTreeData.h b/Engine/Plugins/Developer/ProjectLauncher/Source/ProjectLauncher/Public/ProfileTree/LaunchProfileTreeData.h
--- a/Engine/Plugins/Developer/ProjectLauncher/Source/ProjectLauncher/Public/ProfileTree/LaunchProfileTreeData.h
+++ b/Engine/Plugins/Developer/ProjectLauncher/Source/ProjectLauncher/Public/ProfileTree/LaunchProfileTreeData.h
@@ -152,7 +152,7 @@ namespace ProjectLauncher
 
 		/* Node details */
 		FText Name;													// label for the node
-		int32 SortOrder = 0;										// sort order for headings
+		int32 SortOrder;											// sort order for headings
 		int32 Depth = 0;											// tree depth
 		TSharedPtr<SWidget> Widget;									// UI widget
 		FCallbacks Callbacks;										// generic callbacks
diff --git a/Engine/Plugins/Performance/AutomatedPerfTesting/Source/AutomatedPerfTesting/Private/LaunchExtension/AutomatedPerfTestLaunchExtension.cpp b/Engine/Plugins/Performance/AutomatedPerfTesting/Source/AutomatedPerfTesting/Private/LaunchExtension/AutomatedPerfTestLaunchExtension.cpp
--- a/Engine/Plugins/Performance/AutomatedPerfTesting/Source/AutomatedPerfTesting/Private/LaunchExtension/AutomatedPerfTestLaunchExtension.cpp
+++ b/Engine/Plugins/Performance/AutomatedPerfTesting/Source/AutomatedPerfTesting/Private/LaunchExtension/AutomatedPerfTestLaunchExtension.cpp
@@ -139,11 +139,15 @@ static constexpr const TCHAR* GetTestNodeName(EAutomatedPerfTestType TestType)
 
 void FAutomatedPerfTestLaunchExtensionInstance::OnTestAdded(ILauncherProfileAutomatedTestRef AutomatedTest)
 {
-    constexpr int32 DefaultPriority = 1000;
 	const TCHAR* TestNode = GetTestNodeName(GetTestType());
 
 	AutomatedTest->SetTests(TestNode);
-	AutomatedTest->SetPriority(DefaultPriority);
+}
+
+int32 FAutomatedPerfTestLaunchExtensionInstance::GetTestOrder() const
+{
+    constexpr int32 DefaultOrder = 1000;
+	return DefaultOrder;
 }
 
 const FString FAutomatedPerfTestLaunchExtensionInstance::GetTestInternalName() const
diff --git a/Engine/Plugins/Performance/AutomatedPerfTesting/Source/AutomatedPerfTesting/Private/LaunchExtension/AutomatedPerfTestLaunchExtension.h b/Engine/Plugins/Performance/AutomatedPerfTesting/Source/AutomatedPerfTesting/Private/LaunchExtension/AutomatedPerfTestLaunchExtension.h
--- a/Engine/Plugins/Performance/AutomatedPerfTesting/Source/AutomatedPerfTesting/Private/LaunchExtension/AutomatedPerfTestLaunchExtension.h
+++ b/Engine/Plugins/Performance/AutomatedPerfTesting/Source/AutomatedPerfTesting/Private/LaunchExtension/AutomatedPerfTestLaunchExtension.h
@@ -30,6 +30,7 @@ public:
 	virtual void OnPropertyChanged() override;
 	virtual void OnTestAdded(ILauncherProfileAutomatedTestRef AutomatedTest) override;
 	virtual const FString GetTestInternalName() const override;
+	virtual int32 GetTestOrder() const override;
 
 	virtual void CustomizeTree(ProjectLauncher::FLaunchProfileTreeData& ProfileTreeData) override;
 	virtual void CustomizeAutomatedTestCommandLine(FString& InOutCommandLine) override;
diff --git a/Engine/Source/Developer/LauncherServices/Private/Launcher/LauncherWorker.cpp b/Engine/Source/Developer/LauncherServices/Private/Launcher/LauncherWorker.cpp
--- a/Engine/Source/Developer/LauncherServices/Private/Launcher/LauncherWorker.cpp
+++ b/Engine/Source/Developer/LauncherServices/Private/Launcher/LauncherWorker.cpp
@@ -388,7 +388,7 @@ static FString Join(const TSet<FString>& Tokens, const FString& Delimeter)
 	return Result.RightChop(Delimeter.Len());
 }
 
-FString FLauncherWorker::CreateUATCommand( const ILauncherProfileRef& InProfile, const TArray<FString>& InPlatforms, TArray<FCommandDesc>& OutCommands, FString& CommandStart, bool bForTurnkeyCustomBuild )
+FString FLauncherWorker::CreateUATCommandForBuildCookRun( const ILauncherProfileRef& InProfile, const TArray<FString>& InPlatforms, TArray<FCommandDesc>& OutCommands, FString& CommandStart, bool bForTurnkeyCustomBuild )
 {
 	ILauncherProfileManagerRef LauncherProfileManager = FModuleManager::LoadModuleChecked<ILauncherServicesModule>("LauncherServices").GetProfileManager();
 
@@ -1165,38 +1165,11 @@ FString FLauncherWorker::CreateUATCommand( const ILauncherProfileRef& InProfile,
 }
 
 
-FString FLauncherWorker::CreateUATCommandForAutomatedTests( const ILauncherProfileRef& InProfile, const TArray<FString>& InPlatforms, TArray<FCommandDesc>& OutCommands, FString& OutCommandStart )
+FString FLauncherWorker::CreateUATCommandForAutomatedTest( const ILauncherProfileAutomatedTestRef& InAutomatedTest, const ILauncherProfileRef& InProfile, const TArray<FString>& InPlatforms, TArray<FCommandDesc>& OutCommands, FString& OutCommandStart )
 {
-	FString UATCommand;
-
-	if (InProfile->GetAutomatedTests().Num() > 0)
-	{
-		TArray<ILauncherProfileAutomatedTestRef> PrioritizedAutomatedTests = InProfile->GetAutomatedTests();
-		PrioritizedAutomatedTests.Sort([]( const ILauncherProfileAutomatedTestRef& A, const ILauncherProfileAutomatedTestRef& B)
-		{
-			return A->GetPriority() < B->GetPriority();
-		});
-
-		for (const ILauncherProfileAutomatedTestRef& AutomatedTest : PrioritizedAutomatedTests)
-		{
-			FString AutomatedTestCommand = CreateUATCommandForAutomatedTest(AutomatedTest, InProfile, InPlatforms, OutCommands, OutCommandStart);
-			UATCommand += AutomatedTestCommand + TEXT(" ");
-		}
-
-		FCommandDesc Desc;
-		FText Command = LOCTEXT("LauncherAutoTestDesc", "Automated Testing");
-		Desc.Name = "Automated Testing Task";
-		Desc.Desc = Command.ToString();
-		Desc.EndText = TEXT("********** AUTOMATED TESTS COMPLETED **********"); // @note: this isn't actually logged anywhere - this is the last real task so it will advance automatically
-		OutCommands.Add(Desc);
-	}
+	checkf(InAutomatedTest->GetOrder() >= 0, TEXT("Automated tests should not have negative order"));
 
-	return UATCommand;
-}
 
-
-FString FLauncherWorker::CreateUATCommandForAutomatedTest( const ILauncherProfileAutomatedTestRef& InAutomatedTest, const ILauncherProfileRef& InProfile, const TArray<FString>& InPlatforms, TArray<FCommandDesc>& OutCommands, FString& OutCommandStart )
-{
 	// skip disabled tests
 	if (!InAutomatedTest->IsEnabled())
 	{
@@ -1302,8 +1275,27 @@ FString FLauncherWorker::CreateUATCommandForAutomatedTest( const ILauncherProfil
 		AdditionalCommandLine += TestAdditionalCommandLine;
 	}
 
-	// Dispatch the event which may further customize the command line passed to the automated test.
+	PRAGMA_DISABLE_DEPRECATION_WARNINGS
 	LauncherProfileManager->OnPostProcessAutomatedTestCommandLine().Broadcast(InAutomatedTest, InProfile, AdditionalCommandLine);
+	PRAGMA_ENABLE_DEPRECATION_WARNINGS
+
+	// Dispatch the event which may further customize the command line passed to the automated test.
+	FUATCommandPostProcessParameters PostProcessParameters =
+	{
+		.bShouldBeSkipped = true,
+		.Profile = InProfile,
+		.UATCommand = InAutomatedTest,
+		.UATCommandLine = AdditionalCommandLine,
+		.Commands = OutCommands,
+		.CommandStart = OutCommandStart,
+	};
+	LauncherProfileManager->OnPostProcessUATCommandLine().Broadcast(PostProcessParameters);
+	if (PostProcessParameters.bShouldBeSkipped)
+	{
+		UE_LOG(LogLauncherProfile, Log, TEXT("skipping automated test %s"), InAutomatedTest->GetInternalName() );
+		return FString();
+	}
+
 	if (!AdditionalCommandLine.IsEmpty())
 	{
 		AutomatedTestCommand += TEXT(" ");
@@ -1313,13 +1305,74 @@ FString FLauncherWorker::CreateUATCommandForAutomatedTest( const ILauncherProfil
 	return AutomatedTestCommand;
 }
 
+
+FString FLauncherWorker::CreateUATCommand( const ILauncherProfileRef& InProfile, ILauncherProfileUATCommandRef& InUATCommand, const TArray<FString>& InPlatforms, TArray<FCommandDesc>& OutCommands, FString& OutCommandStart )
+{
+	// skip disabled tests
+	if (!InUATCommand->IsEnabled())
+	{
+		return FString();
+	}
+
+	// special case for automated tests
+	ILauncherProfileAutomatedTestPtr AutomatedTest = InUATCommand->AsAutomatedTest();
+	if (AutomatedTest.IsValid())
+	{
+		return CreateUATCommandForAutomatedTest(AutomatedTest.ToSharedRef(), InProfile, InPlatforms, OutCommands, OutCommandStart);
+	}
+
+
+
+	// basic command parameters
+	FString UATCommandLine = FString::Printf( TEXT("%s -utf8output -nop4 "), *InUATCommand->GetUATCommand() );
+
+	FString AdditionalCommandLine = InUATCommand->GetAdditionalCommandLine();
+	if (!AdditionalCommandLine.IsEmpty())
+	{
+		AdditionalCommandLine.ReplaceInline(TEXT("\""), TEXT("\\\""));
+
+		UATCommandLine += TEXT(" ");
+		UATCommandLine += AdditionalCommandLine;
+	}
+
+	// Dispatch the event which may further customize the command line passed to the automated test.
+	ILauncherProfileManagerRef LauncherProfileManager = FModuleManager::LoadModuleChecked<ILauncherServicesModule>("LauncherServices").GetProfileManager();
+
+	FUATCommandPostProcessParameters PostProcessParameters =
+	{
+		.bShouldBeSkipped = true,
+		.Profile = InProfile,
+		.UATCommand = InUATCommand,
+		.UATCommandLine = UATCommandLine,
+		.Commands = OutCommands,
+		.CommandStart = OutCommandStart,
+	};
+	LauncherProfileManager->OnPostProcessUATCommandLine().Broadcast(PostProcessParameters);
+	if (PostProcessParameters.bShouldBeSkipped)
+	{
+		UE_LOG(LogLauncherProfile, Log, TEXT("skipping UAT command %s"), InUATCommand->GetInternalName() );
+		return FString();
+	}
+
+	if (!AdditionalCommandLine.IsEmpty())
+	{
+		UATCommandLine += TEXT(" ");
+		UATCommandLine += AdditionalCommandLine;
+	}
+
+	return UATCommandLine;
+
+}
+
+
+
 FString FLauncherWorker::MakeBuildCookRunParamsForProjectCustomBuild(const ILauncherProfileRef& InProfile, const TArray<FString>& InPlatforms)
 {
 	// get the basic UAT command line
 	FLauncherWorker TempWorker(NoInit);
 	TArray<FCommandDesc> Commands;
 	FString CommandStart;
-	FString BuildCookRunParams = TempWorker.CreateUATCommand(InProfile, InPlatforms, Commands, CommandStart, true);
+	FString BuildCookRunParams = TempWorker.CreateUATCommandForBuildCookRun(InProfile, InPlatforms, Commands, CommandStart, true);
 
 	// optional project
 	if (InProfile->HasProjectSpecified())
@@ -1450,19 +1503,57 @@ void FLauncherWorker::CreateAndExecuteTasks( const ILauncherProfileRef& InProfil
 		TSharedPtr<FLauncherTask> WaitTask = MakeShareable(new FWaitForCookInEditorToFinish());
 		AddTask(WaitTask);
 	}
+
+	// gather all ordered UAT commands
+	TArray<ILauncherProfileUATCommandRef> UATCommands = InProfile->GetUATCommands();
+	UATCommands.Sort( []( const ILauncherProfileUATCommandRef& A, const ILauncherProfileUATCommandRef& B)
+	{
+		return A->GetOrder() < B->GetOrder();
+	});
+	int UATCommandIndex = 0;
+
+
 	TArray<FCommandDesc> Commands;
 	FString StartString;
-	FString UATCommand = CreateUATCommand(InProfile, Platforms, Commands, StartString);
-	
-	// have Turnkey 
-	FString TurnkeyCommand;
+
+	FString PreCommand;
 	if (Profile->ShouldUpdateDeviceFlash() && DeviceGroup->GetDeviceIDs().Num() >= 0)
 	{
-		TurnkeyCommand = FString::Printf(TEXT("Turnkey -command=VerifySdk -type=Flash -device=%s -UpdateIfNeeded -utf8output -WaitForUATMutex %s"), *FString::Join(DeviceGroup->GetDeviceIDs(), TEXT("+")), *ITurnkeyIOModule::Get().GetUATParams());
+		PreCommand = FString::Printf(TEXT("Turnkey -command=VerifySdk -type=Flash -device=%s -UpdateIfNeeded -utf8output -WaitForUATMutex %s"), *FString::Join(DeviceGroup->GetDeviceIDs(), TEXT("+")), *ITurnkeyIOModule::Get().GetUATParams());
+		PreCommand += TEXT(" ");
+	}
+
+	// add pre-BuildCookRun commands
+	while ( UATCommandIndex < UATCommands.Num() && UATCommands[UATCommandIndex]->GetOrder() < 0)
+	{
+		ILauncherProfileUATCommandRef UATCommand = UATCommands[UATCommandIndex++];
+		PreCommand += CreateUATCommand(InProfile, UATCommand, Platforms, Commands, StartString);
+		PreCommand += TEXT(" ");
+	}
+
+	// main BuildCookRun command
+	FString MainUATCommand = CreateUATCommandForBuildCookRun(InProfile, Platforms, Commands, StartString);
+	
+	// add post-BuildCookRun commands
+	FString PostCommand;
+	while ( UATCommandIndex < UATCommands.Num())
+	{
+		ILauncherProfileUATCommandRef UATCommand = UATCommands[UATCommandIndex++];
+		PostCommand += CreateUATCommand(InProfile, UATCommand, Platforms, Commands, StartString);
+		PostCommand += TEXT(" ");
+	}
+
+	// wait for completion of automated tests
+	if (Profile->GetAutomatedTests().Num() > 0)
+	{
+		FCommandDesc Desc;
+		FText Command = LOCTEXT("LauncherAutoTestDesc", "Automated Testing");
+		Desc.Name = "Automated Testing Task";
+		Desc.Desc = Command.ToString();
+		Desc.EndText = TEXT("********** AUTOMATED TESTS COMPLETED **********"); // @note: this isn't actually logged anywhere - this is the last real task so it will advance automatically
+		Commands.Add(Desc);
 	}
 
-	// have Automated tests
-	FString PostCommand = CreateUATCommandForAutomatedTests( InProfile, Platforms, Commands, StartString );
 
 	// wait for completion of UAT
 	{
@@ -1475,7 +1566,7 @@ void FLauncherWorker::CreateAndExecuteTasks( const ILauncherProfileRef& InProfil
 	}
 
 
-	TSharedPtr<FLauncherTask> LaunchTask = MakeShareable(new FLauncherUATTask(UATCommand, TEXT("Launch Task"), TEXT("Launching UAT..."), ReadPipe, WritePipe, InProfile->GetEditorExe(), ProcHandle, this, StartString, TurnkeyCommand, PostCommand));
+	TSharedPtr<FLauncherTask> LaunchTask = MakeShareable(new FLauncherUATTask(MainUATCommand, TEXT("Launch Task"), TEXT("Launching UAT..."), ReadPipe, WritePipe, InProfile->GetEditorExe(), ProcHandle, this, StartString, PreCommand.TrimStartAndEnd(), PostCommand.TrimStartAndEnd()));
 	AddTask(LaunchTask);
 	for (int32 Index = 0; Index < Commands.Num(); ++Index)
 	{
diff --git a/Engine/Source/Developer/LauncherServices/Private/Launcher/LauncherWorker.h b/Engine/Source/Developer/LauncherServices/Private/Launcher/LauncherWorker.h
--- a/Engine/Source/Developer/LauncherServices/Private/Launcher/LauncherWorker.h
+++ b/Engine/Source/Developer/LauncherServices/Private/Launcher/LauncherWorker.h
@@ -11,13 +11,6 @@ class FLauncherTask;
 class ITargetDeviceProxyManager;
 class ITargetDeviceProxy;
 
-struct FCommandDesc
-{
-	FString Name;
-	FString Desc;
-	FString EndText;
-};
-
 /**
  * Implements the launcher's worker thread.
  */
@@ -123,9 +116,10 @@ protected:
 	void OnTaskStarted(const FString& TaskName);
 	void OnTaskCompleted(const FString& TaskName);
 
-	FString CreateUATCommand( const ILauncherProfileRef& InProfile, const TArray<FString>& InPlatforms, TArray<FCommandDesc>& OutCommands, FString& CommandStart, bool bForTurnkeyCustomBuild = false );
+	FString CreateUATCommandForBuildCookRun( const ILauncherProfileRef& InProfile, const TArray<FString>& InPlatforms, TArray<FCommandDesc>& OutCommands, FString& CommandStart, bool bForTurnkeyCustomBuild = false );
+
+	FString CreateUATCommand( const ILauncherProfileRef& InProfile, ILauncherProfileUATCommandRef& InUATCommand, const TArray<FString>& InPlatforms, TArray<FCommandDesc>& OutCommands, FString& OutCommandStart );
 
-	FString CreateUATCommandForAutomatedTests( const ILauncherProfileRef& InProfile, const TArray<FString>& InPlatforms, TArray<FCommandDesc>& OutCommands, FString& OutCommandStart );
 	FString CreateUATCommandForAutomatedTest( const ILauncherProfileAutomatedTestRef& InAutomatedTest, const ILauncherProfileRef& InProfile, const TArray<FString>& InPlatforms, TArray<FCommandDesc>& OutCommands, FString& OutCommandStart );
 
 	//start listening to the device discovery routine
diff --git a/Engine/Source/Developer/LauncherServices/Private/Profiles/LauncherProfile.cpp b/Engine/Source/Developer/LauncherServices/Private/Profiles/LauncherProfile.cpp
--- a/Engine/Source/Developer/LauncherServices/Private/Profiles/LauncherProfile.cpp
+++ b/Engine/Source/Developer/LauncherServices/Private/Profiles/LauncherProfile.cpp
@@ -1,9 +1,17 @@
 // Copyright Epic Games, Inc. All Rights Reserved.
 
 #include "LauncherProfile.h"
+#include "Profiles/LauncherProfileUATCommand.h"
+#include "Profiles/LauncherProfileAutomatedTest.h"
 
 #define LOCTEXT_NAMESPACE "SProjectLauncherValidation"
 
+
+const TCHAR* FLauncherProfileUATCommand::TypeName = TEXT("UATCommand");
+const TCHAR* FLauncherProfileAutomatedTest::TypeName = TEXT("AutomatedTest");
+
+
+
 FString LexToStringLocalized(ELauncherProfileValidationErrors::Type Value)
 {
 	static_assert(ELauncherProfileValidationErrors::Count == 31, "GetLocalizedValidationErrorMessage() needs to be updated to account for modified enum values.");
diff --git a/Engine/Source/Developer/LauncherServices/Private/Profiles/LauncherProfile.h b/Engine/Source/Developer/LauncherServices/Private/Profiles/LauncherProfile.h
--- a/Engine/Source/Developer/LauncherServices/Private/Profiles/LauncherProfile.h
+++ b/Engine/Source/Developer/LauncherServices/Private/Profiles/LauncherProfile.h
@@ -19,6 +19,7 @@
 #include "ITargetDeviceServicesModule.h"
 #include "GameProjectHelper.h"
 #include "Profiles/LauncherProfileLaunchRole.h"
+#include "Profiles/LauncherProfileUATCommand.h"
 #include "Profiles/LauncherProfileAutomatedTest.h"
 #include "PlatformInfo.h"
 #include "TargetReceipt.h"
@@ -67,6 +68,7 @@ enum ELauncherVersion
 	LAUNCHERSERVICES_FASTITERATE = 42,
 	LAUNCHERSERVICES_REMOVEDINCREMENTALCOOKMODE = 43,
 	LAUNCHERSERVICES_ADDMULTIPLEBUILDTARGETS = 44,
+	LAUNCHERSERVICES_ADDUATCOMMANDS = 45,
 	//ADD NEW STUFF HERE
 
 
@@ -1235,11 +1237,20 @@ public:
 
 		if (Version >= LAUNCHERSERVICES_AUTOMATEDTESTS)
 		{
-			// serialize automated tests
-			if (Archive.IsLoading())
+			TArray<TSharedPtr<ILauncherProfileAutomatedTest>> AutomatedTests;
+			if (Archive.IsSaving())
 			{
-				AutomatedTests.Reset();
+				for (const ILauncherProfileUATCommandRef& UATCommand : UATCommands)
+				{
+					ILauncherProfileAutomatedTestPtr AutomatedTest = UATCommand->AsAutomatedTest();
+					if (AutomatedTest)
+					{
+						AutomatedTests.Add(AutomatedTest);
+					}
+				}
 			}
+
+			// serialize automated tests
 			int NumAutomatedTests = AutomatedTests.Num();
 			Archive << NumAutomatedTests;
 			for( int32 TestIndex = 0; TestIndex < NumAutomatedTests; TestIndex++)
@@ -1256,6 +1267,14 @@ PRAGMA_DISABLE_DEPRECATION_WARNINGS
 			}
 PRAGMA_ENABLE_DEPRECATION_WARNINGS
 
+			if (Archive.IsLoading())
+			{
+				for (TSharedPtr<ILauncherProfileAutomatedTest> AutomatedTest : AutomatedTests)
+				{
+					UATCommands.Add(AutomatedTest.ToSharedRef());
+				}
+			}
+
 			Archive << AutomatedTestBuildPath;
 			Archive << bWantAutomatedTestBuild;
 		}
@@ -1465,16 +1484,17 @@ PRAGMA_ENABLE_DEPRECATION_WARNINGS
 			Writer.WriteArrayEnd();
 		}
 
-		// serialize the automated tests
-		Writer.WriteArrayStart("AutomatedTests");
-		for ( ILauncherProfileAutomatedTestRef AutomatedTest : AutomatedTests)
-		{
-			AutomatedTest->Save(Writer);
-		}
-		Writer.WriteArrayEnd();
+		// serialize the automated test info
 		Writer.WriteValue("AutomatedTestBuildPath", AutomatedTestBuildPath);
 		Writer.WriteValue("WantAutomatedTestBuild", bWantAutomatedTestBuild);
 
+		// serialize the UAT commands
+		Writer.WriteArrayStart("UATCommands");
+		for ( const ILauncherProfileUATCommandRef& UATCommand : UATCommands)
+		{
+			UATCommand->Save(Writer);
+		}
+		Writer.WriteArrayEnd();
 
 		// write out the UAT project params
 		SaveUATParams(Writer);
@@ -2011,6 +2031,7 @@ PRAGMA_ENABLE_DEPRECATION_WARNINGS
 		{
 			return false;
 		}
+		UATCommands.Reset();
 
 		FGuid::Parse(Object.GetStringField(TEXT("Id")), Id);
 		Name = Object.GetStringField(TEXT("Name"));
@@ -2237,15 +2258,18 @@ PRAGMA_ENABLE_DEPRECATION_WARNINGS
 
 		if (Version >= LAUNCHERSERVICES_AUTOMATEDTESTS)
 		{
-			AutomatedTests.Reset();
-			const TArray<TSharedPtr<FJsonValue>>* Values = nullptr;
-			if (Object.TryGetArrayField(TEXT("AutomatedTests"), Values))
+			if (Version < LAUNCHERSERVICES_ADDUATCOMMANDS)
 			{
-				for (TSharedPtr<FJsonValue> Value : *Values)
+				const TArray<TSharedPtr<FJsonValue>>* Values = nullptr;
+				if (Object.TryGetArrayField(TEXT("AutomatedTests"), Values))
 				{
-					AutomatedTests.Add(MakeShared<FLauncherProfileAutomatedTest>(*Value->AsObject().Get()));
+					for (TSharedPtr<FJsonValue> Value : *Values)
+					{
+						UATCommands.Add(MakeShared<FLauncherProfileAutomatedTest>(*Value->AsObject().Get()));
+					}
 				}
 			}
+
 			AutomatedTestBuildPath = Object.GetStringField(TEXT("AutomatedTestBuildPath"));
 			bWantAutomatedTestBuild = Object.GetBoolField(TEXT("WantAutomatedTestBuild"));
 		}
@@ -2289,6 +2313,33 @@ PRAGMA_ENABLE_DEPRECATION_WARNINGS
 		}
 
 
+		if (Version >= LAUNCHERSERVICES_ADDUATCOMMANDS)
+		{
+			const TArray<TSharedPtr<FJsonValue>>* Values = nullptr;
+			if (Object.TryGetArrayField(TEXT("UATCommands"), Values))
+			{
+				for (TSharedPtr<FJsonValue> Value : *Values)
+				{
+					const FJsonObject& JsonObject = *Value->AsObject().Get();
+					FString UATCommandType = JsonObject.GetStringField(TEXT("Type"));
+					if (UATCommandType == FLauncherProfileUATCommand::TypeName)
+					{
+						UATCommands.Add(MakeShared<FLauncherProfileUATCommand>(JsonObject));
+					}
+					else if (UATCommandType == FLauncherProfileAutomatedTest::TypeName)
+					{
+						UATCommands.Add(MakeShared<FLauncherProfileAutomatedTest>(JsonObject));
+					}
+					else
+					{
+						UE_LOG(LogLauncherProfile, Error, TEXT("unknown UAT command type: %s"), *UATCommandType);
+						return false;
+					}
+				}
+			}
+		}
+
+
 		// load the default launch role
 		TSharedPtr<FJsonObject> Role = Object.GetObjectField(TEXT("DefaultRole"));
 		DefaultLaunchRole->Load(*(Role.Get()));
@@ -2444,10 +2495,11 @@ PRAGMA_ENABLE_DEPRECATION_WARNINGS
 		bIsDeviceASimulator = false;
 		bMakeBinaryConfig = false;
 
-		AutomatedTests.Reset();
 		AutomatedTestBuildPath.Reset();
 		bWantAutomatedTestBuild = false;
 
+		UATCommands.Reset();
+
 		RefreshValidBuildTargets();
 		Validate();
 	}
@@ -3150,19 +3202,25 @@ PRAGMA_ENABLE_DEPRECATION_WARNINGS
 		return CustomBoolProperties;
 	}
 
-	virtual const TArray<ILauncherProfileAutomatedTestRef>& GetAutomatedTests() const override
+	virtual const TArray<ILauncherProfileAutomatedTestRef> GetAutomatedTests() const override
 	{
+		TArray<ILauncherProfileAutomatedTestRef> AutomatedTests;
+		for (const ILauncherProfileUATCommandRef& UATCommand : UATCommands)
+		{
+			ILauncherProfileAutomatedTestPtr AutomatedTest = UATCommand->AsAutomatedTest();
+			if (AutomatedTest)
+			{
+				AutomatedTests.Add(AutomatedTest.ToSharedRef());
+			}
+		}
+
 		return AutomatedTests;
 	}
 
 	virtual ILauncherProfileAutomatedTestPtr GetAutomatedTest( const TCHAR* InternalName ) const override
 	{
-		const ILauncherProfileAutomatedTestRef* TestPtr = AutomatedTests.FindByPredicate([InternalName]( const ILauncherProfileAutomatedTestRef& Test )
-		{
-			return FCString::Stricmp( InternalName, Test->GetInternalName()) == 0;
-		});
-
-		return TestPtr ? TestPtr->ToSharedPtr() : nullptr;
+		ILauncherProfileUATCommandPtr UATCommand = GetUATCommand(InternalName);
+		return UATCommand ? UATCommand->AsAutomatedTest() : nullptr;
 	}
 
 	virtual ILauncherProfileAutomatedTestRef FindOrAddAutomatedTest( const TCHAR* InternalName ) override
@@ -3171,7 +3229,7 @@ PRAGMA_ENABLE_DEPRECATION_WARNINGS
 		if (AutomatedTest == nullptr)
 		{
 			AutomatedTest = MakeShared<FLauncherProfileAutomatedTest>(InternalName);
-			AutomatedTests.Add(AutomatedTest.ToSharedRef());
+			UATCommands.Add(AutomatedTest.ToSharedRef());
 		}
 
 		return AutomatedTest.ToSharedRef();
@@ -3179,11 +3237,7 @@ PRAGMA_ENABLE_DEPRECATION_WARNINGS
 
 	virtual void RemoveAutomatedTest( const TCHAR* InternalName ) override
 	{
-		ILauncherProfileAutomatedTestPtr AutomatedTest = GetAutomatedTest(InternalName);
-		if (AutomatedTest != nullptr)
-		{
-			AutomatedTests.Remove(AutomatedTest.ToSharedRef());
-		}
+		RemoveUATCommand(InternalName);
 	}
 
 	virtual void SetAutomatedTestBuildPath( const FString& InAutomatedTestBuildPath )
@@ -3204,6 +3258,50 @@ PRAGMA_ENABLE_DEPRECATION_WARNINGS
 		return bWantAutomatedTestBuild;
 	}
 
+
+	virtual const TArray<ILauncherProfileUATCommandRef>& GetUATCommands() const override
+	{
+		return UATCommands;
+	}
+
+	virtual ILauncherProfileUATCommandPtr GetUATCommand( const TCHAR* InternalName ) const
+	{
+		const ILauncherProfileUATCommandRef* CommandPtr = UATCommands.FindByPredicate([InternalName]( const ILauncherProfileUATCommandRef& Test )
+		{
+			return FCString::Stricmp( InternalName, Test->GetInternalName()) == 0;
+		});
+
+		return CommandPtr ? CommandPtr->ToSharedPtr() : nullptr;
+
+	}
+
+	virtual ILauncherProfileUATCommandRef FindOrAddUATCommand( const TCHAR* InternalName )
+	{
+		ILauncherProfileUATCommandPtr UATCommand = GetUATCommand(InternalName);
+		if (UATCommand == nullptr)
+		{
+			UATCommand = MakeShared<FLauncherProfileUATCommand>(InternalName);
+			UATCommands.Add(UATCommand.ToSharedRef());
+		}
+
+		return UATCommand.ToSharedRef();
+	}
+
+	virtual void RemoveUATCommand( const TCHAR* InternalName )
+	{
+		ILauncherProfileUATCommandPtr UATCommand = GetUATCommand(InternalName);
+		if (UATCommand != nullptr)
+		{
+			UATCommands.Remove(UATCommand.ToSharedRef());
+		}
+	}
+
+
+
+
+
+
+
 	virtual TArray<FString> GetBuildTargets() const
 	{
 		if (BuildTargetSpecified)
@@ -3902,10 +4000,12 @@ private:
 	TMap<FString, bool>   CustomBoolProperties;
 
 	// Automated tests
-	TArray<ILauncherProfileAutomatedTestRef> AutomatedTests;
 	FString AutomatedTestBuildPath;
 	bool bWantAutomatedTestBuild;
 
+	// Extra UAT commands
+	TArray<ILauncherProfileUATCommandRef> UATCommands;
+
 private:
 
 	// cook in the editor callbacks (not valid for any other cook mode)
diff --git a/Engine/Source/Developer/LauncherServices/Private/Profiles/LauncherProfileAutomatedTest.h b/Engine/Source/Developer/LauncherServices/Private/Profiles/LauncherProfileAutomatedTest.h
--- a/Engine/Source/Developer/LauncherServices/Private/Profiles/LauncherProfileAutomatedTest.h
+++ b/Engine/Source/Developer/LauncherServices/Private/Profiles/LauncherProfileAutomatedTest.h
@@ -6,9 +6,11 @@
 #include "ILauncherProfileAutomatedTest.h"
 #include "Dom/JsonObject.h"
 
-class FLauncherProfileAutomatedTest final : public ILauncherProfileAutomatedTest
+class FLauncherProfileAutomatedTest final : public ILauncherProfileAutomatedTest, public TSharedFromThis<FLauncherProfileAutomatedTest>
 {
 public:
+	static const TCHAR* TypeName;
+
 	FLauncherProfileAutomatedTest( const TCHAR* InInternalName ) 
 		: InternalName(InInternalName)
 	{
@@ -64,13 +66,13 @@ public:
 		return bEnabled;
 	}
 
-	virtual void SetPriority( int32 InPriority ) override
+	virtual void SetOrder( int32 InOrder ) override
 	{
-		Priority = InPriority;
+		Order = InOrder;
 	}
-	virtual int32 GetPriority() const override
+	virtual int32 GetOrder() const override
 	{
-		return Priority;
+		return Order;
 	}
 
 	virtual const TCHAR* GetInternalName() const override
@@ -78,6 +80,12 @@ public:
 		return *InternalName;
 	}
 
+	virtual ILauncherProfileAutomatedTestPtr AsAutomatedTest() override
+	{
+		return AsShared().ToSharedPtr();
+	}
+
+
 	virtual void Load(const FJsonObject& Object) override
 	{
 		InternalName          = Object.GetStringField(TEXT("InternalName"));
@@ -85,18 +93,19 @@ public:
 		Tests                 = Object.GetStringField(TEXT("Tests"));
 		AdditionalCommandLine = Object.GetStringField(TEXT("AdditionalCommandLine"));
 		bEnabled              = Object.GetBoolField(TEXT("Enabled"));
-		Priority              = Object.GetIntegerField(TEXT("Priority"));
-
+		Order                 = Object.GetIntegerField(TEXT("Order"));
+		Order = FMath::Max(1, Order);
 	}
 	virtual void Save(TJsonWriter<>& Writer) override
 	{
 		Writer.WriteObjectStart();
+		Writer.WriteValue(TEXT("Type"),                  TypeName);
 		Writer.WriteValue(TEXT("InternalName"),          InternalName);
 		Writer.WriteValue(TEXT("UATCommand"),            UATCommand);
 		Writer.WriteValue(TEXT("Tests"),                 Tests);
 		Writer.WriteValue(TEXT("AdditionalCommandLine"), AdditionalCommandLine);
 		Writer.WriteValue(TEXT("Enabled"),               bEnabled);
-		Writer.WriteValue(TEXT("Priority"),              Priority);
+		Writer.WriteValue(TEXT("Order"),                 Order);
 		Writer.WriteObjectEnd();
 	}
 
@@ -108,7 +117,7 @@ public:
 				<< Tests
 				<< AdditionalCommandLine
 				<< bEnabled
-				<< Priority
+				<< Order
 		;
 	}
 
@@ -119,7 +128,7 @@ protected:
 	FString AdditionalCommandLine;
 
 	bool bEnabled = true;
-	int32 Priority = 0; // longer tests should generally have lower priority
+	int32 Order = 1; // longer tests should generally have a higher order
 
 	FString InternalName;
 };
diff --git a/Engine/Source/Developer/LauncherServices/Private/Profiles/LauncherProfileManager.h b/Engine/Source/Developer/LauncherServices/Private/Profiles/LauncherProfileManager.h
--- a/Engine/Source/Developer/LauncherServices/Private/Profiles/LauncherProfileManager.h
+++ b/Engine/Source/Developer/LauncherServices/Private/Profiles/LauncherProfileManager.h
@@ -91,11 +91,17 @@ public:
 		return ProcessProcessLaunchCommandLine;
 	}
 
+	UE_DEPRECATED(5.8, "Use OnPostProcessUATCommandLine")
 	virtual FOnLauncherProfilePostProcessAutomatedTestCommandLine& OnPostProcessAutomatedTestCommandLine( ) override
 	{
 		return PostProcessAutomatedTestCommandLine;
 	}
 
+	virtual FOnLauncherProfilePostProcessUATCommandLine& OnPostProcessUATCommandLine( ) override
+	{
+		return PostProcessUATCommandLine;
+	}
+
 	virtual void RemoveDeviceGroup( const ILauncherDeviceGroupRef& DeviceGroup ) override;
 
 	virtual void SaveDeviceGroups() override;
@@ -229,4 +235,6 @@ private:
 	// Holds a delegate to be involked to modify the command line for an automated test before launching
 	FOnLauncherProfilePostProcessAutomatedTestCommandLine PostProcessAutomatedTestCommandLine;
 
+	// Holds a delegate to be involked to modify the command line for an extra UAT command before launching
+	FOnLauncherProfilePostProcessUATCommandLine PostProcessUATCommandLine;
 };
diff --git a/Engine/Source/Developer/LauncherServices/Private/Profiles/LauncherProfileUATCommand.h b/Engine/Source/Developer/LauncherServices/Private/Profiles/LauncherProfileUATCommand.h
new file mode 100644
--- /dev/null
+++ b/Engine/Source/Developer/LauncherServices/Private/Profiles/LauncherProfileUATCommand.h
@@ -0,0 +1,105 @@
+// Copyright Epic Games, Inc. All Rights Reserved.
+
+#pragma once
+
+#include "CoreMinimal.h"
+#include "ILauncherProfileUATCommand.h"
+#include "ILauncherProfileAutomatedTest.h"
+#include "Dom/JsonObject.h"
+
+class FLauncherProfileUATCommand final : public ILauncherProfileUATCommand, public TSharedFromThis<FLauncherProfileUATCommand>
+{
+public:
+	static const TCHAR* TypeName;
+
+	FLauncherProfileUATCommand( const TCHAR* InInternalName ) 
+		: InternalName(InInternalName)
+	{
+	}
+
+	FLauncherProfileUATCommand(const FJsonObject& Object)
+	{
+		Load(Object);
+	}
+
+	virtual void SetUATCommand( const TCHAR* InUATCommand ) override
+	{
+		UATCommand = InUATCommand;
+	}
+
+	virtual const FString& GetUATCommand() const override
+	{
+		return UATCommand;
+	}
+
+	virtual void SetAdditionalCommandLine( const TCHAR* InAdditionalCommandLine ) override
+	{
+		AdditionalCommandLine = InAdditionalCommandLine;
+	}
+	virtual const FString& GetAdditionalCommandLine() const override
+	{
+		return AdditionalCommandLine;
+	}
+
+	virtual void SetEnabled( bool bInEnabled ) override
+	{
+		bEnabled = bInEnabled;
+	}
+	virtual bool IsEnabled() const override
+	{
+		return bEnabled;
+	}
+
+	virtual void SetOrder( int32 InOrder ) override
+	{
+		check(InOrder != 0);
+		Order = InOrder;
+	}
+	virtual int32 GetOrder() const override
+	{
+		return Order;
+	}
+
+	virtual const TCHAR* GetInternalName() const override
+	{
+		return *InternalName;
+	}
+
+	virtual ILauncherProfileAutomatedTestPtr AsAutomatedTest() override
+	{
+		return nullptr;
+	}
+
+
+	virtual void Load(const FJsonObject& Object) override
+	{
+		InternalName          = Object.GetStringField(TEXT("InternalName"));
+		UATCommand            = Object.GetStringField(TEXT("UATCommand"));
+		AdditionalCommandLine = Object.GetStringField(TEXT("AdditionalCommandLine"));
+		bEnabled              = Object.GetBoolField(TEXT("Enabled"));
+		Order                 = Object.GetIntegerField(TEXT("Order"));
+	}
+	virtual void Save(TJsonWriter<>& Writer) override
+	{
+		Writer.WriteObjectStart();
+		Writer.WriteValue(TEXT("Type"),                  TypeName);
+		Writer.WriteValue(TEXT("InternalName"),          InternalName);
+		Writer.WriteValue(TEXT("UATCommand"),            UATCommand);
+		Writer.WriteValue(TEXT("AdditionalCommandLine"), AdditionalCommandLine);
+		Writer.WriteValue(TEXT("Enabled"),               bEnabled);
+		Writer.WriteValue(TEXT("Order"),                 Order);
+		Writer.WriteObjectEnd();
+	}
+
+
+protected:
+	FString UATCommand = TEXT("RunUnreal");
+	FString AdditionalCommandLine;
+
+	bool bEnabled = true;
+	int32 Order = -100000; // negative order items will run before the main BuildCookRun, positive order items will run afterwards
+
+	FString InternalName;
+};
+
+
diff --git a/Engine/Source/Developer/LauncherServices/Public/ILauncherProfile.h b/Engine/Source/Developer/LauncherServices/Public/ILauncherProfile.h
--- a/Engine/Source/Developer/LauncherServices/Public/ILauncherProfile.h
+++ b/Engine/Source/Developer/LauncherServices/Public/ILauncherProfile.h
@@ -1657,10 +1657,38 @@ public:
 	virtual const TArray<FString>& GetEditorArchitectures() const = 0;
 
 
+
+
+
+	
+	/**
+	 * Returns a list of all UAT commands in the profile
+	 */
+	virtual const TArray<ILauncherProfileUATCommandRef>& GetUATCommands() const = 0;
+
+	/**
+	 * Returns an existing UAT command by internal name
+	 */
+	virtual ILauncherProfileUATCommandPtr GetUATCommand( const TCHAR* InternalName ) const = 0;
+
+	/**
+	 * Returns an UAT command by internal name, creating it if it doesn't already exist
+	 */
+	virtual ILauncherProfileUATCommandRef FindOrAddUATCommand( const TCHAR* InternalName ) = 0;
+
+	/**
+	 * Removes an UAT command by internal name
+	 */
+	virtual void RemoveUATCommand( const TCHAR* InternalName ) = 0;
+
+
+
+
+
 	/**
 	 * Returns a list of all automated tests in the profile
 	 */
-	virtual const TArray<ILauncherProfileAutomatedTestRef>& GetAutomatedTests() const = 0;
+	virtual const TArray<ILauncherProfileAutomatedTestRef> GetAutomatedTests() const = 0;
 
 	/**
 	 * Returns an existing automated test by internal name
diff --git a/Engine/Source/Developer/LauncherServices/Public/ILauncherProfileAutomatedTest.h b/Engine/Source/Developer/LauncherServices/Public/ILauncherProfileAutomatedTest.h
--- a/Engine/Source/Developer/LauncherServices/Public/ILauncherProfileAutomatedTest.h
+++ b/Engine/Source/Developer/LauncherServices/Public/ILauncherProfileAutomatedTest.h
@@ -3,43 +3,29 @@
 #pragma once
 
 #include "CoreMinimal.h"
-#include "Serialization/JsonWriter.h"
+#include "ILauncherProfileUATCommand.h"
 
 class FJsonObject;
-class ILauncherProfileLaunchRole;
 
 /**
  * Automated test parameters
  */
-class ILauncherProfileAutomatedTest
+class ILauncherProfileAutomatedTest : public ILauncherProfileUATCommand
 {
 public:
 	//@todo: comments
 
-	// recommendation for InternalName is is "ExtensionName.InternalName"
-	virtual void SetUATCommand( const TCHAR* UATCommand ) = 0;
-	virtual const FString& GetUATCommand() const = 0;
-
 	virtual void SetTests( const TCHAR* Tests ) = 0; //-Test=xxx,yyy
 	virtual const FString& GetTests() const = 0;
 
-	virtual void SetAdditionalCommandLine( const TCHAR* AdditionalCommandLine ) = 0;
-	virtual const FString& GetAdditionalCommandLine() const = 0;
-
-	virtual void SetEnabled( bool bEnabled ) = 0;
-	virtual bool IsEnabled() const = 0;
-
-	virtual void SetPriority( int32 Priority ) = 0; // longer tests should generally have lower priority
-	virtual int32 GetPriority() const = 0;
-
-	virtual const TCHAR* GetInternalName() const = 0;
-
 	virtual void Load(const FJsonObject& Object) = 0;
 	virtual void Save(TJsonWriter<>& Writer) = 0;
 
 	UE_DEPRECATED(5.8, "please use the Json Save/Load API")
 	virtual void Serialize( FArchive& Archive ) = 0;
 
+	UE_DEPRECATED(5.8, "Use SetOrder")
+	inline void SetPriority( int32 Priority ) { SetOrder(INT32_MAX-Priority); }
 
 	virtual ~ILauncherProfileAutomatedTest() = default;
 };
diff --git a/Engine/Source/Developer/LauncherServices/Public/ILauncherProfileManager.h b/Engine/Source/Developer/LauncherServices/Public/ILauncherProfileManager.h
--- a/Engine/Source/Developer/LauncherServices/Public/ILauncherProfileManager.h
+++ b/Engine/Source/Developer/LauncherServices/Public/ILauncherProfileManager.h
@@ -7,6 +7,7 @@
 
 class ILauncherProfileManager;
 class ILauncherProfileWizard;
+struct FCommandDesc;
 
 /** Type definition for shared pointers to instances of ILauncherProfileManager. */
 typedef TSharedPtr<class ILauncherProfileManager> ILauncherProfileManagerPtr;
@@ -14,6 +15,19 @@ typedef TSharedPtr<class ILauncherProfileManager> ILauncherProfileManagerPtr;
 /** Type definition for shared references to instances of ILauncherProfileManager. */
 typedef TSharedRef<class ILauncherProfileManager> ILauncherProfileManagerRef;
 
+/** Parameters for customizing a UAT command via a callback */
+struct FUATCommandPostProcessParameters
+{
+	bool bShouldBeSkipped;
+
+	const ILauncherProfileRef Profile;
+	const ILauncherProfileUATCommandRef UATCommand;
+	FString& UATCommandLine;
+	
+	TArray<FCommandDesc>& Commands;
+	FString& CommandStart;
+};
+
 
 /**
  * Declares a delegate to be invoked when a device group was added to a profile manager.
@@ -64,6 +78,11 @@ DECLARE_MULTICAST_DELEGATE_ThreeParams(FOnLauncherProfilePostProcessCommandLine,
  */
 DECLARE_MULTICAST_DELEGATE_ThreeParams(FOnLauncherProfilePostProcessAutomatedTestCommandLine, const ILauncherProfileAutomatedTestRef&, const ILauncherProfileRef&, FString& );
 
+/**
+ * Declares a delegate to be invoked to post-process a command line before launching a UAT command
+ */
+DECLARE_MULTICAST_DELEGATE_OneParam(FOnLauncherProfilePostProcessUATCommandLine, FUATCommandPostProcessParameters& );
+
 
 /** Type definition for shared references to instances of ILauncherProfile. */
 typedef TSharedPtr<class ILauncherProfileWizard> ILauncherProfileWizardPtr;
@@ -431,6 +450,13 @@ public:
 	 */
 	virtual FOnLauncherProfilePostProcessAutomatedTestCommandLine& OnPostProcessAutomatedTestCommandLine( ) = 0;
 
+	/**
+	 * Returns a delegate that is invoked to post-process the command line before a UAT command is launched
+	 * 
+	 * @return The delegate.
+	 */
+	virtual FOnLauncherProfilePostProcessUATCommandLine& OnPostProcessUATCommandLine( ) = 0;
+
 public:
 
 	/** Virtual destructor. */
diff --git a/Engine/Source/Developer/LauncherServices/Public/ILauncherProfileUATCommand.h b/Engine/Source/Developer/LauncherServices/Public/ILauncherProfileUATCommand.h
new file mode 100644
--- /dev/null
+++ b/Engine/Source/Developer/LauncherServices/Public/ILauncherProfileUATCommand.h
@@ -0,0 +1,43 @@
+// Copyright Epic Games, Inc. All Rights Reserved.
+
+#pragma once
+
+#include "CoreMinimal.h"
+#include "Serialization/JsonWriter.h"
+
+class FJsonObject;
+class ILauncherProfileAutomatedTest;
+
+/**
+ * UAT command parameters
+ */
+class ILauncherProfileUATCommand
+{
+public:
+	//@todo: comments
+
+	virtual void SetUATCommand( const TCHAR* UATCommand ) = 0;
+	virtual const FString& GetUATCommand() const = 0;
+
+	virtual void SetAdditionalCommandLine( const TCHAR* AdditionalCommandLine ) = 0;
+	virtual const FString& GetAdditionalCommandLine() const = 0;
+
+	virtual void SetEnabled( bool bEnabled ) = 0;
+	virtual bool IsEnabled() const = 0;
+
+	virtual void SetOrder( int32 Order ) = 0; // negative order items will run before the main BuildCookRun, positive order items will run afterwards
+	virtual int32 GetOrder() const = 0;
+
+	// recommendation for InternalName is is "ExtensionName.InternalName"
+	virtual const TCHAR* GetInternalName() const = 0;
+
+	virtual void Load(const FJsonObject& Object) = 0;
+	virtual void Save(TJsonWriter<>& Writer) = 0;
+
+	virtual TSharedPtr<ILauncherProfileAutomatedTest> AsAutomatedTest() = 0;
+
+	virtual ~ILauncherProfileUATCommand() = default;
+};
+
+typedef TSharedPtr<ILauncherProfileUATCommand> ILauncherProfileUATCommandPtr;
+typedef TSharedRef<ILauncherProfileUATCommand> ILauncherProfileUATCommandRef;
diff --git a/Engine/Source/Developer/LauncherServices/Public/ILauncherWorker.h b/Engine/Source/Developer/LauncherServices/Public/ILauncherWorker.h
--- a/Engine/Source/Developer/LauncherServices/Public/ILauncherWorker.h
+++ b/Engine/Source/Developer/LauncherServices/Public/ILauncherWorker.h
@@ -45,6 +45,14 @@ DECLARE_MULTICAST_DELEGATE_ThreeParams(FOnLaunchCompletedDelegate, bool, double,
 /** Delegate used to notify when the launch was canceled */
 DECLARE_MULTICAST_DELEGATE_OneParam(FOnLaunchCanceledDelegate, double);
 
+/** For task list customization */
+struct FCommandDesc
+{
+	FString Name;
+	FString Desc;
+	FString EndText;
+};
+
 
 /**
  * Interface for launcher worker threads.
