Launcher Profile support for multiple build cook run commands per-profile.

- majority of ILauncherProfile members now route to an inner BuildCookRun. These API functions will likely be deprecated in an upcoming change.

#rnx
#Synced-CL 49676362

diff --git a/Engine/Source/Developer/LauncherServices/Private/Launcher/LauncherUATTask.h b/Engine/Source/Developer/LauncherServices/Private/Launcher/LauncherUATTask.h
--- a/Engine/Source/Developer/LauncherServices/Private/Launcher/LauncherUATTask.h
+++ b/Engine/Source/Developer/LauncherServices/Private/Launcher/LauncherUATTask.h
@@ -17,7 +17,7 @@ class FLauncherUATTask
 {
 public:
 
-	FLauncherUATTask( const FString& InCommandLine, const FString& InName, const FString& InDesc, void* InReadPipe, void* InWritePipe, const FString& InEditorExe, FProcHandle& InProcessHandle, ILauncherWorker* InWorker, const FString& InCommandEnd, const FString& InBuildCookRunPreCommand, const FString& InBuildCookRunPostCommand)
+	FLauncherUATTask( const FString& InCommandLine, const FString& InName, const FString& InDesc, void* InReadPipe, void* InWritePipe, const FString& InEditorExe, FProcHandle& InProcessHandle, ILauncherWorker* InWorker, const FString& InCommandEnd)
 		: FLauncherTask(InName, InDesc)
 
 		, CommandLine(InCommandLine)
@@ -26,8 +26,6 @@ public:
 		, EditorExe(InEditorExe)
 		, ProcessHandle(InProcessHandle)
 		, CommandText(InCommandEnd)
-		, PreCommand(InBuildCookRunPreCommand)
-		, PostCommand(InBuildCookRunPostCommand)
 	{
 		InWorker->OnOutputReceived().AddRaw(this, &FLauncherUATTask::HandleOutputReceived);
 	}
@@ -55,16 +53,10 @@ protected:
 #endif
 
 		// base UAT command arguments
-		static const FString ConfigStrings[] = { TEXT("Unknown"), TEXT("Debug"), TEXT("DebugGame"), TEXT("Development"), TEXT("Shipping"), TEXT("Test") };
 		FString UATCommandLine;
 		FString ProjectPath = *ChainState.Profile->GetProjectPath();
 		ProjectPath = FPaths::ConvertRelativePathToFull(ProjectPath);
-		UATCommandLine = FString::Printf(TEXT("-ScriptsForProject=\"%s\" %s BuildCookRun -project=\"%s\" -noP4 -clientconfig=%s -serverconfig=%s"),
-			*ProjectPath,
-			*PreCommand,
-			*ProjectPath,
-			LexToString(ChainState.Profile->GetBuildConfiguration()),
-			LexToString(ChainState.Profile->GetBuildConfiguration()));
+		UATCommandLine = FString::Printf(TEXT("-ScriptsForProject=\"%s\" -noP4"), *ProjectPath);
 
 		// we expect to pass -nocompile to UAT here as we generally expect UAT to be fully compiled. Besides, installed builds don't even have the source to compile UAT scripts.
 		// Only allow UAT to compile scripts dynamically if we pass -development or we have the IsBuildingUAT property set, the latter of which should not allow itself to be set in installed situations.
@@ -73,6 +65,7 @@ protected:
 		// we never want to build the editor when launching from the editor or running with an installed engine (which can't rebuild itself)
 		UATCommandLine += GIsEditor || FApp::IsEngineInstalled() ? TEXT(" -nocompileeditor") : TEXT("");
 		UATCommandLine += FApp::IsEngineInstalled() ? TEXT(" -installed") : TEXT("");
+		UATCommandLine += TEXT(" -utf8output");
 
 		// specify the path to the editor exe if necessary
 		if(EditorExe.Len() > 0)
@@ -81,15 +74,9 @@ protected:
 		}
 
 		// specialized command arguments for this particular task
+		UATCommandLine += TEXT(" ");
 		UATCommandLine += CommandLine;
 
-		// add any post command line
-		if (!PostCommand.IsEmpty())
-		{
-			UATCommandLine += TEXT(" ");
-			UATCommandLine += PostCommand;
-		}
-
 
 		// launch UAT and monitor its progress
 		ProcessHandle = FPlatformProcess::CreateProc(*(ExecutablePath / Executable), *UATCommandLine, false, true, true, NULL, 0, *ExecutablePath, WritePipe);
@@ -136,15 +123,9 @@ private:
 	// process handle
 	FProcHandle& ProcessHandle;
 
-	// The editor executable that UAT should use
+	// The end text marker for this UAT command
 	FString CommandText;
 
-	// a UAT command to run before BuildCookRun (in the same instance as the BuildCookRun)
-	FString PreCommand;
-
-	// a UAT command to run after BuildCookRun (in the same insance as the BuildCookRun)
-	FString PostCommand;
-
 	bool EndTextFound = false;
 };
 
diff --git a/Engine/Source/Developer/LauncherServices/Private/Launcher/LauncherWorker.cpp b/Engine/Source/Developer/LauncherServices/Private/Launcher/LauncherWorker.cpp
--- a/Engine/Source/Developer/LauncherServices/Private/Launcher/LauncherWorker.cpp
+++ b/Engine/Source/Developer/LauncherServices/Private/Launcher/LauncherWorker.cpp
@@ -20,6 +20,7 @@
 #include "PlatformInfo.h"
 #include "Misc/ConfigCacheIni.h"
 #include "Profiles/LauncherProfile.h"
+#include "Profiles/LauncherProfileBuildCookRun.h"
 #include "DerivedDataCacheInterface.h"
 
 
@@ -388,22 +389,35 @@ static FString Join(const TSet<FString>& Tokens, const FString& Delimeter)
 	return Result.RightChop(Delimeter.Len());
 }
 
-FString FLauncherWorker::CreateUATCommandForBuildCookRun( const ILauncherProfileRef& InProfile, const TArray<FString>& InPlatforms, TArray<FCommandDesc>& OutCommands, FString& CommandStart, bool bForTurnkeyCustomBuild )
+FString FLauncherWorker::CreateUATCommandForBuildCookRun( const ILauncherProfileBuildCookRunRef& InBuildCookRun, const ILauncherProfileRef& InProfile, TArray<FCommandDesc>& OutCommands, FString& CommandStart, const FCommandAdjustments& Adjustments, bool bForTurnkeyCustomBuild )
 {
 	ILauncherProfileManagerRef LauncherProfileManager = FModuleManager::LoadModuleChecked<ILauncherServicesModule>("LauncherServices").GetProfileManager();
 
 	CommandStart = TEXT("");
-	FString UATCommand = bForTurnkeyCustomBuild ? TEXT("") : TEXT(" -utf8output");
+
+	
+	bool bPreferCookedPlatforms = (Adjustments.bPreferCookedPlatforms || InBuildCookRun->GetCookMode() == ELauncherProfileCookModes::ByTheBook || InBuildCookRun->ShouldBuild());
+	TArray<FString> FinalPlatforms = FinalizePlatforms(InBuildCookRun->GetCookedPlatforms(), InBuildCookRun->GetDeployedDeviceGroup(), bPreferCookedPlatforms);
+
+	bool bRun = (InBuildCookRun->GetLaunchMode() != ELauncherProfileLaunchModes::DoNotLaunch && !Adjustments.bDisableRun);
+
+	// base UAT command arguments
+	static const FString ConfigStrings[] = { TEXT("Unknown"), TEXT("Debug"), TEXT("DebugGame"), TEXT("Development"), TEXT("Shipping"), TEXT("Test") };
+	FString UATCommand = FString::Printf(TEXT("BuildCookRun -project=\"%s\" -clientconfig=%s -serverconfig=%s"), // @todo: don't have to always add -clientconfig & -serverconfig
+		*FPaths::ConvertRelativePathToFull(InProfile->GetProjectPath()),
+		LexToString(InBuildCookRun->GetBuildConfiguration()),
+		LexToString(InBuildCookRun->GetBuildConfiguration()));
+
 	FGuid SessionId(FGuid::NewGuid());
-	FString InitialMap = InProfile->GetDefaultLaunchRole()->GetInitialMap();
-	if (InitialMap.IsEmpty() && InProfile->GetCookedMaps().Num() == 1)
+	FString InitialMap = InBuildCookRun->GetDefaultLaunchRole()->GetInitialMap();
+	if (InitialMap.IsEmpty() && InBuildCookRun->GetCookedMaps().Num() == 1)
 	{
-		InitialMap = InProfile->GetCookedMaps()[0];
+		InitialMap = InBuildCookRun->GetCookedMaps()[0];
 	}
 
 	// staging directory
 	FString StageDirectory = TEXT("");
-	auto PackageDirectory = InProfile->GetPackageDirectory();
+	auto PackageDirectory = InBuildCookRun->GetPackageDirectory();
 	if (PackageDirectory.Len() > 0)
 	{
 		StageDirectory += FString::Printf(TEXT(" -stagingdirectory=\"%s\""), *PackageDirectory);
@@ -420,10 +434,10 @@ FString FLauncherWorker::CreateUATCommandForBuildCookRun( const ILauncherProfile
 	TSet<EBuildTargetType> BuildTargets;
 	
 	bool bUATClosesAfterLaunch = false;
-	for (int32 PlatformIndex = 0; PlatformIndex < InPlatforms.Num(); ++PlatformIndex)
+	for (int32 PlatformIndex = 0; PlatformIndex < FinalPlatforms.Num(); ++PlatformIndex)
 	{
 		// Platform info for the given platform
-		const PlatformInfo::FTargetPlatformInfo* PlatformInfo = PlatformInfo::FindPlatformInfo(FName(*InPlatforms[PlatformIndex]));
+		const PlatformInfo::FTargetPlatformInfo* PlatformInfo = PlatformInfo::FindPlatformInfo(FName(*FinalPlatforms[PlatformIndex]));
 		if (bForTurnkeyCustomBuild && PlatformInfo == nullptr)
 		{
 			continue;
@@ -479,9 +493,9 @@ FString FLauncherWorker::CreateUATCommandForBuildCookRun( const ILauncherProfile
 			OptionalParams += TEXT(" -noclient");
 		}
 
-		if (InProfile->GetServerArchitectures().Num() > 0)
+		if (InBuildCookRun->GetServerArchitectures().Num() > 0)
 		{
-			ServerCommand += TEXT(" -serverarchitecture=") + FString::Join(InProfile->GetServerArchitectures(), TEXT("+") );
+			ServerCommand += TEXT(" -serverarchitecture=") + FString::Join(InBuildCookRun->GetServerArchitectures(), TEXT("+") );
 		}
 	}
 	bool bSetClientArchitecture = false;
@@ -505,12 +519,12 @@ FString FLauncherWorker::CreateUATCommandForBuildCookRun( const ILauncherProfile
 		UATCommand += (TEXT(" -cookflavor=") + Join(OptionalCookFlavors, TEXT("+")));
 	}
 
-	if (InProfile->GetBuildTargets().Num() > 0 && !InProfile->GetBuildTargets().Contains(FString()))
+	if (InBuildCookRun->GetBuildTargets().Num() > 0 && !InBuildCookRun->GetBuildTargets().Contains(FString()))
 	{
-		UATCommand += TEXT(" -target=") + FString::Join(InProfile->GetBuildTargets(), TEXT("+"));
+		UATCommand += TEXT(" -target=") + FString::Join(InBuildCookRun->GetBuildTargets(), TEXT("+"));
 	}
 
-	if (InProfile->IsDeviceASimulator())
+	if (InBuildCookRun->IsDeviceASimulator())
 	{
 		if (Platforms.Contains(TEXT("IOS")))
 		{
@@ -520,9 +534,9 @@ FString FLauncherWorker::CreateUATCommandForBuildCookRun( const ILauncherProfile
 		// TODO: add tvOS and VisionOS simulators below
 	}
 
-	if (bSetClientArchitecture && InProfile->GetClientArchitectures().Num() > 0)
+	if (bSetClientArchitecture && InBuildCookRun->GetClientArchitectures().Num() > 0)
 	{
-		UATCommand += TEXT(" -clientarchitecture=") + FString::Join(InProfile->GetClientArchitectures(), TEXT("+") );
+		UATCommand += TEXT(" -clientarchitecture=") + FString::Join(InBuildCookRun->GetClientArchitectures(), TEXT("+") );
 	}
 
 
@@ -530,7 +544,7 @@ FString FLauncherWorker::CreateUATCommandForBuildCookRun( const ILauncherProfile
 	FString DeviceNames = TEXT("");
 	FString DeviceCommand = TEXT("");
 	FString RoleCommands = TEXT("");
-	ILauncherDeviceGroupPtr DeviceGroup = InProfile->GetDeployedDeviceGroup();
+	ILauncherDeviceGroupPtr DeviceGroup = InBuildCookRun->GetDeployedDeviceGroup();
 
 	bool bVsyncAdded = false;
 
@@ -566,7 +580,7 @@ FString FLauncherWorker::CreateUATCommandForBuildCookRun( const ILauncherProfile
 		}
 		else
 		{
-			RoleCommands = InProfile->GetDefaultLaunchRole()->GetUATCommandLine();
+			RoleCommands = InBuildCookRun->GetDefaultLaunchRole()->GetUATCommandLine();
 		}
 	}
 
@@ -576,6 +590,7 @@ FString FLauncherWorker::CreateUATCommandForBuildCookRun( const ILauncherProfile
 	}
 
 	// game command line
+	//@fixme: this is destined for UECommandLine.txt ... probably want initial map (if set) but not -Messaging ?
 	FString CommandLine = FString::Printf(TEXT(" -cmdline=\"%s%s\""),
 		*InitialMap, 
 		bForTurnkeyCustomBuild ? TEXT("") : TEXT(" -Messaging") );
@@ -595,81 +610,83 @@ FString FLauncherWorker::CreateUATCommandForBuildCookRun( const ILauncherProfile
 	}
 #endif	// WITH_EDITOR
 
-	// to reduce UECommandLine.txt churn (timestamp causing extra work), for LaunchOn (ie iterative deploy) we use a single session guid
-	if (InProfile->GetDeploymentMode() == ELauncherProfileDeploymentModes::CopyToDevice && InProfile->IsDeployingIncrementally())
-	{
-		static FGuid StaticGuid(FGuid::NewGuid());
-		SessionId = StaticGuid;
-	}
-
-	// additional commands to be sent to the commandline
-	FString SessionCommands;
-	if (!bForTurnkeyCustomBuild)
-	{
-		FString SessionName = InProfile->GetName().Replace(TEXT("\'"), TEXT("_")).Replace(TEXT("\'"), TEXT("_")).Replace(TEXT("&"), TEXT("_"));
-		FString SessionOwner = FString(FPlatformProcess::UserName(false)).Replace(TEXT("\'"), TEXT("_")).Replace(TEXT("\'"), TEXT("_"));;
-		SessionCommands = FString::Printf(TEXT("-SessionId=%s -SessionOwner='%s' -SessionName='%s' "),
-		*SessionId.ToString(),
-		*SessionOwner,
-		*SessionName);
-	}
-
-	// allow external items to adjust the command line
-	FString ProfileAdditionalCommandLineParameters = InProfile->GetAdditionalCommandLineParameters();
-	LauncherProfileManager->OnPostProcessLaunchCommandLine().Broadcast(InProfile, EBuildTargetType::Unknown, ProfileAdditionalCommandLineParameters);
-
-	FString AdditionalCommandLine;
-	TStringBuilder<64> InheritableGameOptions;
-	ECommandLineArgumentFlags CommandLineContextFlags = ECommandLineArgumentFlags::ClientContext | ECommandLineArgumentFlags::ServerContext;
 	const bool bOnlyInherited = true; // Only getting inherited, not explicitly set subprocess commandline args due to issues when passing
-									  // the -Multiprocess explicitly set commandline argument which prevents configs from saving/writing
-	FCommandLine::BuildSubprocessCommandLine(CommandLineContextFlags, bOnlyInherited, InheritableGameOptions);
-	FString EscapedInheritableGameOptions(InheritableGameOptions);
-	EscapedInheritableGameOptions.ReplaceInline(TEXT("\""), TEXT("\\\""));
-	FString AddCmdLine = FString::Printf(TEXT("%s%s %s %s%s"),
-		*SessionCommands,
-		*RoleCommands,
-		*LocalizationCommands,
-		*ProfileAdditionalCommandLineParameters,
-		*EscapedInheritableGameOptions);
-	if (!AddCmdLine.TrimStartAndEnd().IsEmpty()) // don't bother with -addcmdline if it's empty
+										// the -Multiprocess explicitly set commandline argument which prevents configs from saving/writing
+	if (bRun)
 	{
-		AdditionalCommandLine = FString::Printf(TEXT(" -addcmdline=\"%s\""), *AddCmdLine);
-	}
+		// to reduce UECommandLine.txt churn (timestamp causing extra work), for LaunchOn (ie iterative deploy) we use a single session guid
+		if (InBuildCookRun->GetDeploymentMode() == ELauncherProfileDeploymentModes::CopyToDevice && InBuildCookRun->IsDeployingIncrementally())
+		{
+			static FGuid StaticGuid(FGuid::NewGuid());
+			SessionId = StaticGuid;
+		}
 
-	for (EBuildTargetType BuildTarget : BuildTargets)
-	{
-		FString BuildTargetCmdLine = InProfile->GetAdditionalTargetCommandLineParameters(BuildTarget);
-		LauncherProfileManager->OnPostProcessLaunchCommandLine().Broadcast(InProfile, BuildTarget, BuildTargetCmdLine);
-		if (BuildTargetCmdLine.TrimStartAndEnd().IsEmpty())
+		// additional commands to be sent to the commandline
+		FString SessionCommands;
+		if (!bForTurnkeyCustomBuild)
 		{
-			continue;		
+			FString SessionName = InProfile->GetName().Replace(TEXT("\'"), TEXT("_")).Replace(TEXT("\'"), TEXT("_")).Replace(TEXT("&"), TEXT("_"));
+			FString SessionOwner = FString(FPlatformProcess::UserName(false)).Replace(TEXT("\'"), TEXT("_")).Replace(TEXT("\'"), TEXT("_"));;
+			SessionCommands = FString::Printf(TEXT("-SessionId=%s -SessionOwner='%s' -SessionName='%s' "),
+			*SessionId.ToString(),
+			*SessionOwner,
+			*SessionName);
 		}
-		BuildTargetCmdLine.ReplaceInline(TEXT("\""), TEXT("\\\""));
 
-		switch (BuildTarget)
+		// allow external items to adjust the command line
+		FString ProfileAdditionalCommandLineParameters = InBuildCookRun->GetAdditionalCommandLineParameters();
+		LauncherProfileManager->OnPostProcessLaunchCommandLine().Broadcast(InProfile, EBuildTargetType::Unknown, ProfileAdditionalCommandLineParameters);
+
+		FString AdditionalCommandLine;
+		TStringBuilder<64> InheritableGameOptions;
+		ECommandLineArgumentFlags CommandLineContextFlags = ECommandLineArgumentFlags::ClientContext | ECommandLineArgumentFlags::ServerContext;
+		FCommandLine::BuildSubprocessCommandLine(CommandLineContextFlags, bOnlyInherited, InheritableGameOptions);
+		FString EscapedInheritableGameOptions(InheritableGameOptions);
+		EscapedInheritableGameOptions.ReplaceInline(TEXT("\""), TEXT("\\\""));
+		FString AddCmdLine = FString::Printf(TEXT("%s%s %s %s%s"),
+			*SessionCommands,
+			*RoleCommands,
+			*LocalizationCommands,
+			*ProfileAdditionalCommandLineParameters,
+			*EscapedInheritableGameOptions);
+		if (!AddCmdLine.TrimStartAndEnd().IsEmpty()) // don't bother with -addcmdline if it's empty
 		{
-			case EBuildTargetType::Server:
+			AdditionalCommandLine = FString::Printf(TEXT(" -addcmdline=\"%s\""), *AddCmdLine);
+		}
+
+		for (EBuildTargetType BuildTarget : BuildTargets)
+		{
+			FString BuildTargetCmdLine = InBuildCookRun->GetAdditionalTargetCommandLineParameters(BuildTarget);
+			LauncherProfileManager->OnPostProcessLaunchCommandLine().Broadcast(InProfile, BuildTarget, BuildTargetCmdLine);
+			if (BuildTargetCmdLine.TrimStartAndEnd().IsEmpty())
 			{
-				AdditionalCommandLine += FString::Printf(TEXT(" -servercmdline=\"%s\""), *BuildTargetCmdLine);
+				continue;		
 			}
-			break;
+			BuildTargetCmdLine.ReplaceInline(TEXT("\""), TEXT("\\\""));
 
-			case EBuildTargetType::Client:
-			case EBuildTargetType::Game:
+			switch (BuildTarget)
 			{
-				AdditionalCommandLine += FString::Printf(TEXT(" -clientcmdline=\"%s\""), *BuildTargetCmdLine);
+				case EBuildTargetType::Server:
+				{
+					AdditionalCommandLine += FString::Printf(TEXT(" -servercmdline=\"%s\""), *BuildTargetCmdLine);
+				}
+				break;
+
+				case EBuildTargetType::Client:
+				case EBuildTargetType::Game:
+				{
+					AdditionalCommandLine += FString::Printf(TEXT(" -clientcmdline=\"%s\""), *BuildTargetCmdLine);
+				}
+				break;
 			}
-			break;
 		}
 	}
 
 
-
 	// map list
 	FString MapList;
-	const TArray<FString>& CookedMaps = InProfile->GetCookedMaps();
-	if (CookedMaps.Num() > 0 && (InProfile->GetCookMode() == ELauncherProfileCookModes::ByTheBook || InProfile->GetCookMode() == ELauncherProfileCookModes::ByTheBookInEditor))
+	const TArray<FString>& CookedMaps = InBuildCookRun->GetCookedMaps();
+	if (CookedMaps.Num() > 0 && (InBuildCookRun->GetCookMode() == ELauncherProfileCookModes::ByTheBook || InBuildCookRun->GetCookMode() == ELauncherProfileCookModes::ByTheBookInEditor))
 	{
 		if (!InitialMap.IsEmpty())
 		{
@@ -690,15 +707,15 @@ FString FLauncherWorker::CreateUATCommandForBuildCookRun( const ILauncherProfile
 	// culture list
 	FString CultureList;
 	{
-		const TArray<FString>& CookedCultures = InProfile->GetCookedCultures();
-		if (CookedCultures.Num() > 0 && (InProfile->GetCookMode() == ELauncherProfileCookModes::ByTheBook || InProfile->GetCookMode() == ELauncherProfileCookModes::ByTheBookInEditor))
+		const TArray<FString>& CookedCultures = InBuildCookRun->GetCookedCultures();
+		if (CookedCultures.Num() > 0 && (InBuildCookRun->GetCookMode() == ELauncherProfileCookModes::ByTheBook || InBuildCookRun->GetCookMode() == ELauncherProfileCookModes::ByTheBookInEditor))
 		{
 			CultureList += TEXT(" -CookCultures=");
 			CultureList += FString::Join(CookedCultures, TEXT("+"));
 		}
 	}
 
-	bool bIsBuilding = InProfile->ShouldBuild();
+	bool bIsBuilding = InBuildCookRun->ShouldBuild();
 
 	// build
 	if (bIsBuilding)
@@ -706,7 +723,7 @@ FString FLauncherWorker::CreateUATCommandForBuildCookRun( const ILauncherProfile
 		UATCommand += TEXT(" -build");
 
 		FCommandDesc Desc;
-		FText Command = FText::Format(LOCTEXT("LauncherBuildDesc", "Build game for {0}"), FText::FromString(Platforms.RightChop(1)));
+		FText Command = FText::Format(LOCTEXT("LauncherBuildDesc", "Build game for {0}"), FText::FromString(FString::Join(FinalPlatforms, TEXT("+") )));
 		Desc.Name = "Build Task";
 		Desc.Desc = Command.ToString();
 		Desc.EndText = TEXT("********** BUILD COMMAND COMPLETED **********");
@@ -715,18 +732,18 @@ FString FLauncherWorker::CreateUATCommandForBuildCookRun( const ILauncherProfile
 		// @todo: server
 	}
 	
-	if (InProfile->IsFastIterate())
+	if (InBuildCookRun->IsFastIterate())
 	{
 		UATCommand += TEXT(" -FastIterate -ubtargs=\"-FastIterate\"");
 	}
 
 	// snapshot import
-	if (InProfile->IsImportingZenSnapshot())
+	if (InBuildCookRun->IsImportingZenSnapshot())
 	{
 		UATCommand += TEXT(" -snapshot");
 
 		FCommandDesc Desc;
-		FText Command = FText::Format(LOCTEXT("LauncherSnapshotDesc", "Import Zen snapshot for {0}"), FText::FromString(Platforms.RightChop(1)));
+		FText Command = FText::Format(LOCTEXT("LauncherSnapshotDesc", "Import Zen snapshot for {0}"), FText::FromString(FString::Join(FinalPlatforms, TEXT("+") )));
 		Desc.Name = "Snapshot Import Task";
 		Desc.Desc = Command.ToString();
 		Desc.EndText = TEXT("********** SNAPSHOT IMPORT COMPLETED **********");
@@ -735,7 +752,7 @@ FString FLauncherWorker::CreateUATCommandForBuildCookRun( const ILauncherProfile
 	}
 
 	// cook
-	switch(InProfile->GetCookMode())
+	switch(InBuildCookRun->GetCookMode())
 	{
 	case ELauncherProfileCookModes::ByTheBook:
 		{
@@ -744,19 +761,19 @@ FString FLauncherWorker::CreateUATCommandForBuildCookRun( const ILauncherProfile
 			UATCommand += MapList;
 			UATCommand += CultureList;
 
-			if (InProfile->IsCookingUnversioned())
+			if (InBuildCookRun->IsCookingUnversioned())
 			{
 				UATCommand += TEXT(" -unversionedcookedcontent");
 			}
 
-			if (InProfile->IsEncryptingIniFiles())
+			if (InBuildCookRun->IsEncryptingIniFiles())
 			{
 				UATCommand += TEXT(" -encryptinifiles");
 			}
 
 			TStringBuilder<64> InheritableCookOptions;
 			FCommandLine::BuildSubprocessCommandLine(ECommandLineArgumentFlags::CommandletContext, bOnlyInherited, InheritableCookOptions);
-			FString AdditionalOptions = InProfile->GetCookOptions();
+			FString AdditionalOptions = InBuildCookRun->GetCookOptions();
 			if (!AdditionalOptions.IsEmpty() || (InheritableCookOptions.Len() > 0))
 			{
 				UATCommand += TEXT(" -additionalcookeroptions=\"");
@@ -790,7 +807,7 @@ FString FLauncherWorker::CreateUATCommandForBuildCookRun( const ILauncherProfile
 				UATCommand += TEXT(" -fastcook");
 			}
 
-			if (InProfile->IsUsingZenStore())
+			if (InBuildCookRun->IsUsingZenStore())
 			{
 				UATCommand += TEXT(" -zenstore");
 			}
@@ -804,89 +821,89 @@ FString FLauncherWorker::CreateUATCommandForBuildCookRun( const ILauncherProfile
 				}
 			}
 
-			if (InProfile->IsPackingWithUnrealPak())
+			if (InBuildCookRun->IsPackingWithUnrealPak())
 			{
 				UATCommand += TEXT(" -pak");
 				UATCommand += TEXT(" -stage");
-				if (InProfile->IsUsingIoStore())
+				if (InBuildCookRun->IsUsingIoStore())
 				{
 					UATCommand += TEXT(" -iostore");
 				}
-				if (InProfile->IsCompressed())
+				if (InBuildCookRun->IsCompressed())
 				{
 					UATCommand += TEXT(" -compressed");
 				}
 			}
 
-			if (InProfile->MakeBinaryConfig())
+			if (InBuildCookRun->MakeBinaryConfig())
 			{
 				UATCommand += TEXT(" -makebinaryconfig");
 			}
 
-			if ( InProfile->IsCreatingReleaseVersion() )
+			if ( InBuildCookRun->IsCreatingReleaseVersion() )
 			{
 				UATCommand += TEXT(" -createreleaseversion=");
-				UATCommand += InProfile->GetCreateReleaseVersionName();
+				UATCommand += InBuildCookRun->GetCreateReleaseVersionName();
 				
 			}
 
-			if ( InProfile->IsCreatingDLC() )
+			if ( InBuildCookRun->IsCreatingDLC() )
 			{
 				UATCommand += TEXT(" -dlcname=");
-				UATCommand += InProfile->GetDLCName();
+				UATCommand += InBuildCookRun->GetDLCName();
 			}
 
-			if ( InProfile->IsDLCIncludingEngineContent() )
+			if ( InBuildCookRun->IsDLCIncludingEngineContent() )
 			{
 				UATCommand += TEXT(" -DLCIncludeEngineContent");
 			}
 
-			if ( InProfile->IsGeneratingPatch() )
+			if ( InBuildCookRun->IsGeneratingPatch() )
 			{
 				UATCommand += TEXT(" -generatepatch");
 
-				if ( InProfile->ShouldAddPatchLevel() )
+				if ( InBuildCookRun->ShouldAddPatchLevel() )
 				{
 					UATCommand += TEXT(" -addpatchlevel");
 				}
 			}
 
-			if ( InProfile->IsGeneratingPatch() || 
-				InProfile->IsCreatingReleaseVersion() || 
-				InProfile->IsCreatingDLC() )
+			if ( InBuildCookRun->IsGeneratingPatch() || 
+				InBuildCookRun->IsCreatingReleaseVersion() || 
+				InBuildCookRun->IsCreatingDLC() )
 			{
-				if ( InProfile->GetBasedOnReleaseVersionName().IsEmpty() == false )
+				if ( InBuildCookRun->GetBasedOnReleaseVersionName().IsEmpty() == false )
 				{
 					UATCommand += TEXT(" -basedonreleaseversion=");
-					UATCommand += InProfile->GetBasedOnReleaseVersionName();
+					UATCommand += InBuildCookRun->GetBasedOnReleaseVersionName();
 				}
 
-				if (InProfile->GetOriginalReleaseVersionName().IsEmpty() == false)
+				if (InBuildCookRun->GetOriginalReleaseVersionName().IsEmpty() == false)
 				{
 					UATCommand += TEXT(" -originalreleaseversion=");
-					UATCommand += InProfile->GetOriginalReleaseVersionName();
+					UATCommand += InBuildCookRun->GetOriginalReleaseVersionName();
 				}
 			}
 
-			if (InProfile->IsGeneratingChunks())
+			if (InBuildCookRun->IsGeneratingChunks())
 			{
 				UATCommand += TEXT(" -manifests");
 			}
 
-			if (InProfile->IsGenerateHttpChunkData())
+			if (InBuildCookRun->IsGenerateHttpChunkData())
 			{
-				auto Cmd = FString::Printf(TEXT(" -createchunkinstall -chunkinstalldirectory=\"%s\" -chunkinstallversion=\"%s\""), *InProfile->GetHttpChunkDataDirectory(), *InProfile->GetHttpChunkDataReleaseName());
+				auto Cmd = FString::Printf(TEXT(" -createchunkinstall -chunkinstalldirectory=\"%s\" -chunkinstallversion=\"%s\""), *InBuildCookRun->GetHttpChunkDataDirectory(), *InBuildCookRun->GetHttpChunkDataReleaseName());
 				UATCommand += Cmd;
 			}
 			
 			// Creating a packed DLC requires staging
-			if (InProfile->GetPackagingMode() == ELauncherProfilePackagingModes::DoNotPackage && InProfile->IsCreatingDLC() && InProfile->IsPackingWithUnrealPak())
+			if (InBuildCookRun->GetPackagingMode() == ELauncherProfilePackagingModes::DoNotPackage && InBuildCookRun->IsCreatingDLC() && InBuildCookRun->IsPackingWithUnrealPak())
 			{
 				UATCommand += TEXT(" -stage");
 			}
 
 			FCommandDesc Desc;
-			FText Command = FText::Format(LOCTEXT("LauncherCookDesc", "Cook content for {0}"), FText::FromString(Platforms.RightChop(1)));
+			FText Command = FText::Format(LOCTEXT("LauncherCookDesc", "Cook content for {0}"), FText::FromString(FString::Join(FinalPlatforms, TEXT("+") )));
 			Desc.Name = "Cook Task";
 			Desc.Desc = Command.ToString();
 			Desc.EndText = TEXT("********** COOK COMMAND COMPLETED **********");
@@ -901,7 +918,7 @@ FString FLauncherWorker::CreateUATCommandForBuildCookRun( const ILauncherProfile
 		{
 			UATCommand += TEXT(" -cookonthefly");
 			
-			if (InProfile->IsUsingZenStore())
+			if (InBuildCookRun->IsUsingZenStore())
 			{
 				UATCommand += TEXT(" -zenstore");
 			}
@@ -909,7 +926,7 @@ FString FLauncherWorker::CreateUATCommandForBuildCookRun( const ILauncherProfile
 			TStringBuilder<64> InheritableCookOptions;
 			FCommandLine::BuildSubprocessCommandLine(ECommandLineArgumentFlags::CommandletContext, bOnlyInherited, InheritableCookOptions);
 
-			FString AdditionalOptions = InProfile->GetCookOptions();
+			FString AdditionalOptions = InBuildCookRun->GetCookOptions();
 			if (!AdditionalOptions.IsEmpty() || (InheritableCookOptions.Len() > 0))
 			{
 				UATCommand += TEXT(" -additionalcookeroptions=\"");
@@ -966,7 +983,7 @@ FString FLauncherWorker::CreateUATCommandForBuildCookRun( const ILauncherProfile
 	case ELauncherProfileCookModes::OnTheFlyInEditor:
 		UATCommand += MapList;
 		UATCommand += " -skipcook -cookonthefly -CookInEditor";
-		if (InProfile->IsUsingZenStore())
+		if (InBuildCookRun->IsUsingZenStore())
 		{
 			UATCommand += TEXT(" -zenstore");
 		}
@@ -975,18 +992,18 @@ FString FLauncherWorker::CreateUATCommandForBuildCookRun( const ILauncherProfile
 		UATCommand += MapList;
 		UATCommand += CultureList;
 		UATCommand += TEXT(" -skipcook -CookInEditor"); // don't cook anything the editor is doing it ;)
-		if (InProfile->IsUsingZenStore())
+		if (InBuildCookRun->IsUsingZenStore())
 		{
 			UATCommand += TEXT(" -zenstore");
 		}
-		if (InProfile->IsPackingWithUnrealPak())
+		if (InBuildCookRun->IsPackingWithUnrealPak())
 		{
 			UATCommand += TEXT(" -pak");
-			if (InProfile->IsUsingIoStore())
+			if (InBuildCookRun->IsUsingIoStore())
 			{
 				UATCommand += TEXT(" -iostore");
 			}
-			if (InProfile->IsCompressed())
+			if (InBuildCookRun->IsCompressed())
 			{
 				UATCommand += TEXT(" -compressed");
 			}
@@ -994,18 +1011,18 @@ FString FLauncherWorker::CreateUATCommandForBuildCookRun( const ILauncherProfile
 		break;
 	case ELauncherProfileCookModes::DoNotCook:
 		UATCommand += TEXT(" -skipcook");
-		if (InProfile->IsUsingZenPakStreaming())
+		if (InBuildCookRun->IsUsingZenPakStreaming())
 		{
-			UATCommand += FString::Printf(TEXT(" -skippak -ZenWorkspaceSharePath=\\\"%s\\\" "), *InProfile->GetZenPakStreamingPath() );
+			UATCommand += FString::Printf(TEXT(" -skippak -ZenWorkspaceSharePath=\\\"%s\\\" "), *InBuildCookRun->GetZenPakStreamingPath() );
 		}
-		else if (InProfile->IsImportingZenSnapshot() && InProfile->IsPackingWithUnrealPak())
+		else if (InBuildCookRun->IsImportingZenSnapshot() && InBuildCookRun->IsPackingWithUnrealPak())
 		{
 			UATCommand += TEXT(" -pak");
-			if (InProfile->IsUsingIoStore())
+			if (InBuildCookRun->IsUsingIoStore())
 			{
 				UATCommand += TEXT(" -iostore");
 			}
-			if (InProfile->IsCompressed())
+			if (InBuildCookRun->IsCompressed())
 			{
 				UATCommand += TEXT(" -compressed");
 			}
@@ -1013,41 +1030,41 @@ FString FLauncherWorker::CreateUATCommandForBuildCookRun( const ILauncherProfile
 		break;
 	}
 
-	if ( InProfile->IsForDistribution() )
+	if ( InBuildCookRun->IsForDistribution() )
 	{
 		UATCommand += TEXT(" -distribution");
 	}
 
-	if (InProfile->IsCookingIncrementally())
+	if (InBuildCookRun->IsCookingIncrementally())
 	{
 		UATCommand += TEXT(" -iterativecooking");
 	}
 
-	if ( InProfile->IsIterateSharedCookedBuild() )
+	if ( InBuildCookRun->IsIterateSharedCookedBuild() )
 	{
 		UATCommand += TEXT(" -iteratesharedcookedbuild=usesyncedbuild");
 	}
 
-	if (InProfile->GetSkipCookingEditorContent())
+	if (InBuildCookRun->GetSkipCookingEditorContent())
 	{
 		UATCommand += TEXT(" -SkipCookingEditorContent");
 	}
 
 	FString StageAdditionalCommandLine;
-	if (InProfile->IsUsingIoStore() &&
-		InProfile->GetReferenceContainerGlobalFileName().Len())
+	if (InBuildCookRun->IsUsingIoStore() &&
+		InBuildCookRun->GetReferenceContainerGlobalFileName().Len())
 	{
-		StageAdditionalCommandLine += TEXT(" -ReferenceContainerGlobalFileName=\"") + InProfile->GetReferenceContainerGlobalFileName() + TEXT("\"");
-		if (InProfile->GetReferenceContainerCryptoKeysFileName().Len())
+		StageAdditionalCommandLine += TEXT(" -ReferenceContainerGlobalFileName=\"") + InBuildCookRun->GetReferenceContainerGlobalFileName() + TEXT("\"");
+		if (InBuildCookRun->GetReferenceContainerCryptoKeysFileName().Len())
 		{
-			StageAdditionalCommandLine += TEXT(" -ReferenceContainerCryptoKeys=\"") + InProfile->GetReferenceContainerCryptoKeysFileName() + TEXT("\"") ;
+			StageAdditionalCommandLine += TEXT(" -ReferenceContainerCryptoKeys=\"") + InBuildCookRun->GetReferenceContainerCryptoKeysFileName() + TEXT("\"") ;
 		}
 	}
 
 	// stage/package/deploy
-	if (InProfile->GetDeploymentMode() != ELauncherProfileDeploymentModes::DoNotDeploy)
+	if (InBuildCookRun->GetDeploymentMode() != ELauncherProfileDeploymentModes::DoNotDeploy)
 	{
-		switch (InProfile->GetDeploymentMode())
+		switch (InBuildCookRun->GetDeploymentMode())
 		{
 		case ELauncherProfileDeploymentModes::CopyRepository:
 			{
@@ -1055,10 +1072,10 @@ FString FLauncherWorker::CreateUATCommandForBuildCookRun( const ILauncherProfile
 				UATCommand += CommandLine;
 				UATCommand += StageDirectory;
 				UATCommand += DeviceCommand;
-				UATCommand += AdditionalCommandLine;
+				//UATCommand += AdditionalCommandLine;
 
 				FCommandDesc Desc;
-				FText Command = FText::Format(LOCTEXT("LauncherDeployDesc", "Deploying content for {0}"), FText::FromString(Platforms.RightChop(1)));
+				FText Command = FText::Format(LOCTEXT("LauncherDeployDesc", "Deploying content for {0}"), FText::FromString(FString::Join(FinalPlatforms, TEXT("+") )));
 				Desc.Name = "Deploy Task";
 				Desc.Desc = Command.ToString();
 				Desc.EndText = TEXT("********** DEPLOY COMMAND COMPLETED **********");
@@ -1072,7 +1089,7 @@ FString FLauncherWorker::CreateUATCommandForBuildCookRun( const ILauncherProfile
 
 		case ELauncherProfileDeploymentModes::CopyToDevice:
 			{
-				if (InProfile->IsDeployingIncrementally())
+				if (InBuildCookRun->IsDeployingIncrementally())
 				{
 					UATCommand += " -iterativedeploy";
 				}
@@ -1083,11 +1100,11 @@ FString FLauncherWorker::CreateUATCommandForBuildCookRun( const ILauncherProfile
 				UATCommand += CommandLine;
 				UATCommand += StageDirectory;
 				UATCommand += DeviceCommand;
-				UATCommand += AdditionalCommandLine;
+				//UATCommand += AdditionalCommandLine;
 				UATCommand += StageAdditionalCommandLine;
 
 				FCommandDesc Desc;
-				FText Command = FText::Format(LOCTEXT("LauncherDeployDesc", "Deploying content for {0}"), FText::FromString(Platforms.RightChop(1)));
+				FText Command = FText::Format(LOCTEXT("LauncherDeployDesc", "Deploying content for {0}"), FText::FromString(FString::Join(FinalPlatforms, TEXT("+") )));
 				Desc.Name = "Deploy Task";
 				Desc.Desc = Command.ToString();
 				Desc.EndText = TEXT("********** DEPLOY COMMAND COMPLETED **********");
@@ -1099,41 +1116,24 @@ FString FLauncherWorker::CreateUATCommandForBuildCookRun( const ILauncherProfile
 			}
 			break;
 		}
-
-		// run
-		if (InProfile->GetLaunchMode() != ELauncherProfileLaunchModes::DoNotLaunch && InProfile->GetAutomatedTests().Num() == 0)
-		{
-			UATCommand += TEXT(" -run ");
-
-			FCommandDesc Desc;
-			FText Command = FText::Format(LOCTEXT("LauncherRunDesc", "Launching on {0}"), FText::FromString(DeviceNames.RightChop(1)));
-			Desc.Name = "Run Task";
-			Desc.Desc = Command.ToString();
-			Desc.EndText = TEXT("********** RUN COMMAND COMPLETED **********");
-			OutCommands.Add(Desc);
-			if (CommandStart.Len() == 0)
-			{
-				CommandStart = TEXT("********** RUN COMMAND STARTED **********");
-			}
-		}
 	}
 	else
 	{
-		if (InProfile->IsIncludingPrerequisites())
+		if (InBuildCookRun->IsIncludingPrerequisites())
 		{
 			UATCommand += TEXT(" -prereqs");
 		}
 
-		if (InProfile->GetPackagingMode() == ELauncherProfilePackagingModes::Locally)
+		if (InBuildCookRun->GetPackagingMode() == ELauncherProfilePackagingModes::Locally)
 		{
 			UATCommand += TEXT(" -stage -package");
 			UATCommand += StageDirectory;
 			UATCommand += CommandLine;
-			UATCommand += AdditionalCommandLine;
+			//UATCommand += AdditionalCommandLine;
 			UATCommand += StageAdditionalCommandLine;
 
 			FCommandDesc Desc;
-			FText Command = FText::Format(LOCTEXT("LauncherPackageDesc", "Packaging content for {0}"), FText::FromString(Platforms.RightChop(1)));
+			FText Command = FText::Format(LOCTEXT("LauncherPackageDesc", "Packaging content for {0}"), FText::FromString(FString::Join(FinalPlatforms, TEXT("+") )));
 			Desc.Name = "Package Task";
 			Desc.Desc = Command.ToString();
 			Desc.EndText = TEXT("********** PACKAGE COMMAND COMPLETED **********");
@@ -1144,12 +1144,12 @@ FString FLauncherWorker::CreateUATCommandForBuildCookRun( const ILauncherProfile
 			}
 		}
 
-		if (InProfile->IsArchiving())
+		if (InBuildCookRun->IsArchiving())
 		{
-			UATCommand += FString::Printf(TEXT(" -archive -archivedirectory=\"%s\""), *InProfile->GetArchiveDirectory());
+			UATCommand += FString::Printf(TEXT(" -archive -archivedirectory=\"%s\""), *InBuildCookRun->GetArchiveDirectory());
 
 			FCommandDesc Desc;
-			FText Command = FText::Format(LOCTEXT("LauncherArchiveDesc", "Archiving content for {0}"), FText::FromString(Platforms.RightChop(1)));
+			FText Command = FText::Format(LOCTEXT("LauncherArchiveDesc", "Archiving content for {0}"), FText::FromString(FString::Join(FinalPlatforms, TEXT("+") )));
 			Desc.Name = "Archive Task";
 			Desc.Desc = Command.ToString();
 			Desc.EndText = TEXT("********** ARCHIVE COMMAND COMPLETED **********");
@@ -1161,11 +1161,57 @@ FString FLauncherWorker::CreateUATCommandForBuildCookRun( const ILauncherProfile
 		}
 	}
 
+	// run
+	if (bRun)
+	{
+		UATCommand += TEXT(" -run ");
+
+		FCommandDesc Desc;
+		FText Command = FText::Format(LOCTEXT("LauncherRunDesc", "Launching on {0}"), FText::FromString(DeviceNames.RightChop(1)));
+		Desc.Name = "Run Task";
+		Desc.Desc = Command.ToString();
+		Desc.EndText = TEXT("********** RUN COMMAND COMPLETED **********");
+		OutCommands.Add(Desc);
+		if (CommandStart.Len() == 0)
+		{
+			CommandStart = TEXT("********** RUN COMMAND STARTED **********");
+		}
+	}
+
+
+	FString AdditionalUATCommandLine = InBuildCookRun->GetAdditionalUATCommandLine();
+	if (!AdditionalUATCommandLine.IsEmpty())
+	{
+		AdditionalUATCommandLine.ReplaceInline(TEXT("\""), TEXT("\\\""));
+
+		UATCommand += TEXT(" ");
+		UATCommand += AdditionalUATCommandLine;
+	}
+
+	// Dispatch the event which may further customize the command line passed to the automated test.
+	FUATCommandPostProcessParameters PostProcessParameters =
+	{
+		.bShouldBeSkipped = false, // don't skip by default
+		.Profile = InProfile,
+		.UATCommand = InBuildCookRun,
+		.UATCommandLine = AdditionalUATCommandLine,
+		.Commands = OutCommands,
+		.CommandStart = CommandStart,
+	};
+	LauncherProfileManager->OnPostProcessUATCommandLine().Broadcast(PostProcessParameters);
+	if (PostProcessParameters.bShouldBeSkipped)
+	{
+		UE_LOG(LogLauncherProfile, Log, TEXT("skipping BuildCookRun %s"), InBuildCookRun->GetInternalName() );
+		return FString();
+	}
+
+
+
 	return UATCommand;
 }
 
 
-FString FLauncherWorker::CreateUATCommandForAutomatedTest( const ILauncherProfileAutomatedTestRef& InAutomatedTest, const ILauncherProfileRef& InProfile, const TArray<FString>& InPlatforms, TArray<FCommandDesc>& OutCommands, FString& OutCommandStart )
+FString FLauncherWorker::CreateUATCommandForAutomatedTest( const ILauncherProfileAutomatedTestRef& InAutomatedTest, const ILauncherProfileRef& InProfile, const TArray<FString>& InPlatforms, TArray<FCommandDesc>& OutCommands, FString& OutCommandStart, const FCommandAdjustments& Adjustments )
 {
 	checkf(InAutomatedTest->GetOrder() >= 0, TEXT("Automated tests should not have negative order"));
 
@@ -1268,11 +1314,13 @@ FString FLauncherWorker::CreateUATCommandForAutomatedTest( const ILauncherProfil
 	}
 
 	// Append the tests additional command line parameters
-	FString TestAdditionalCommandLine = InAutomatedTest->GetAdditionalCommandLine();
-	if (!TestAdditionalCommandLine.IsEmpty())
+	FString AdditionalUATCommandLine = InAutomatedTest->GetAdditionalUATCommandLine();
+	if (!AdditionalUATCommandLine.IsEmpty())
 	{
+		AdditionalUATCommandLine.ReplaceInline(TEXT("\""), TEXT("\\\""));
+
 		AdditionalCommandLine += TEXT(" ");
-		AdditionalCommandLine += TestAdditionalCommandLine;
+		AdditionalCommandLine += AdditionalUATCommandLine;
 	}
 
 	PRAGMA_DISABLE_DEPRECATION_WARNINGS
@@ -1282,7 +1330,7 @@ FString FLauncherWorker::CreateUATCommandForAutomatedTest( const ILauncherProfil
 	// Dispatch the event which may further customize the command line passed to the automated test.
 	FUATCommandPostProcessParameters PostProcessParameters =
 	{
-		.bShouldBeSkipped = true,
+		.bShouldBeSkipped = true, // skip tests by default unless something explicitly requets it
 		.Profile = InProfile,
 		.UATCommand = InAutomatedTest,
 		.UATCommandLine = AdditionalCommandLine,
@@ -1306,7 +1354,7 @@ FString FLauncherWorker::CreateUATCommandForAutomatedTest( const ILauncherProfil
 }
 
 
-FString FLauncherWorker::CreateUATCommand( const ILauncherProfileRef& InProfile, ILauncherProfileUATCommandRef& InUATCommand, const TArray<FString>& InPlatforms, TArray<FCommandDesc>& OutCommands, FString& OutCommandStart )
+FString FLauncherWorker::CreateUATCommand( const ILauncherProfileUATCommandRef& InUATCommand, const ILauncherProfileRef& InProfile, const TArray<FString>& InPlatforms, TArray<FCommandDesc>& OutCommands, FString& OutCommandStart, const FCommandAdjustments& Adjustments )
 {
 	// skip disabled tests
 	if (!InUATCommand->IsEnabled())
@@ -1318,21 +1366,27 @@ FString FLauncherWorker::CreateUATCommand( const ILauncherProfileRef& InProfile,
 	ILauncherProfileAutomatedTestPtr AutomatedTest = InUATCommand->AsAutomatedTest();
 	if (AutomatedTest.IsValid())
 	{
-		return CreateUATCommandForAutomatedTest(AutomatedTest.ToSharedRef(), InProfile, InPlatforms, OutCommands, OutCommandStart);
+		return CreateUATCommandForAutomatedTest(AutomatedTest.ToSharedRef(), InProfile, InPlatforms, OutCommands, OutCommandStart, Adjustments);
 	}
 
+	// special case for BuildCookRun
+	ILauncherProfileBuildCookRunPtr BuildCookRun = InUATCommand->AsBuildCookRun();
+	if (BuildCookRun.IsValid())
+	{
+		return CreateUATCommandForBuildCookRun(BuildCookRun.ToSharedRef(), InProfile, OutCommands, OutCommandStart, Adjustments);
+	}
 
 
 	// basic command parameters
-	FString UATCommandLine = FString::Printf( TEXT("%s -utf8output -nop4 "), *InUATCommand->GetUATCommand() );
+	FString UATCommandLine = FString::Printf( TEXT("%s "), *InUATCommand->GetUATCommand() );
 
-	FString AdditionalCommandLine = InUATCommand->GetAdditionalCommandLine();
-	if (!AdditionalCommandLine.IsEmpty())
+	FString AdditionalUATCommandLine = InUATCommand->GetAdditionalUATCommandLine();
+	if (!AdditionalUATCommandLine.IsEmpty())
 	{
-		AdditionalCommandLine.ReplaceInline(TEXT("\""), TEXT("\\\""));
+		AdditionalUATCommandLine.ReplaceInline(TEXT("\""), TEXT("\\\""));
 
 		UATCommandLine += TEXT(" ");
-		UATCommandLine += AdditionalCommandLine;
+		UATCommandLine += AdditionalUATCommandLine;
 	}
 
 	// Dispatch the event which may further customize the command line passed to the automated test.
@@ -1340,7 +1394,7 @@ FString FLauncherWorker::CreateUATCommand( const ILauncherProfileRef& InProfile,
 
 	FUATCommandPostProcessParameters PostProcessParameters =
 	{
-		.bShouldBeSkipped = true,
+		.bShouldBeSkipped = true, // skip general-purpose UAT commands unless something explicitly requetss it
 		.Profile = InProfile,
 		.UATCommand = InUATCommand,
 		.UATCommandLine = UATCommandLine,
@@ -1354,11 +1408,6 @@ FString FLauncherWorker::CreateUATCommand( const ILauncherProfileRef& InProfile,
 		return FString();
 	}
 
-	if (!AdditionalCommandLine.IsEmpty())
-	{
-		UATCommandLine += TEXT(" ");
-		UATCommandLine += AdditionalCommandLine;
-	}
 
 	return UATCommandLine;
 
@@ -1368,11 +1417,19 @@ FString FLauncherWorker::CreateUATCommand( const ILauncherProfileRef& InProfile,
 
 FString FLauncherWorker::MakeBuildCookRunParamsForProjectCustomBuild(const ILauncherProfileRef& InProfile, const TArray<FString>& InPlatforms)
 {
+	// this function can only support one build cook run command (aka "classic launcher profiles")
+	TArray<ILauncherProfileBuildCookRunRef> AllBuildCookRunCmds = InProfile->GetBuildCookRunCommands();
+	if (AllBuildCookRunCmds.Num() == 0)
+	{
+		return FString();
+	}
+
 	// get the basic UAT command line
 	FLauncherWorker TempWorker(NoInit);
 	TArray<FCommandDesc> Commands;
 	FString CommandStart;
-	FString BuildCookRunParams = TempWorker.CreateUATCommandForBuildCookRun(InProfile, InPlatforms, Commands, CommandStart, true);
+	const FCommandAdjustments Adjustments;
+	FString BuildCookRunParams = TempWorker.CreateUATCommandForBuildCookRun(AllBuildCookRunCmds[0], InProfile, Commands, CommandStart, Adjustments, true);
 
 	// optional project
 	if (InProfile->HasProjectSpecified())
@@ -1409,23 +1466,125 @@ FString FLauncherWorker::MakeBuildCookRunParamsForProjectCustomBuild(const ILaun
 
 
 
+ILauncherProfileUATCommandPtr FLauncherWorker::MakeTemporaryUnifiedRunCommand(const ILauncherProfileRef& InProfile)
+{
+	// find all BuildCookRun commands that want to run that we are able to unify
+	TArray<ILauncherProfileBuildCookRunRef> BuildCookRuns;
+	for (const ILauncherProfileBuildCookRunRef& BuildCookRun : InProfile->GetBuildCookRunCommands())
+	{
+		if (BuildCookRun->GetLaunchMode() != ELauncherProfileLaunchModes::DoNotLaunch && !BuildCookRun->IsUsingZenPakStreaming())
+		{
+			BuildCookRuns.Add(BuildCookRun);
+		}
+	}
+
+	// nothing to do if there's nothing to unify
+	if (BuildCookRuns.Num() < 2)
+	{
+		return nullptr;
+	}
+
+	auto CombineCommandLine = []( const FString& A, const FString& B)
+	{
+		if (A.IsEmpty())
+		{
+			return B;
+		}
+		else if (B.IsEmpty())
+		{
+			return A;
+		}
+		else
+		{
+			return A + TEXT(" ") + B;
+		}
+	};
+
 
-/* FLauncherWorker implementation
- *****************************************************************************/
+	ILauncherProfileManagerRef LauncherProfileManager = FModuleManager::LoadModuleChecked<ILauncherServicesModule>("LauncherServices").GetProfileManager();
+	ILauncherProfileBuildCookRunPtr UnifiedRunCommand = MakeShared<FLauncherProfileBuildCookRun>(LauncherProfileManager, nullptr, TEXT("TEMP_UnifiedRunCommand") );
 
-void FLauncherWorker::CreateAndExecuteTasks( const ILauncherProfileRef& InProfile )
-{
-	FPlatformProcess::CreatePipe(ReadPipe, WritePipe);
+	UnifiedRunCommand->SetDeployedDeviceGroup( LauncherProfileManager->CreateUnmanagedDeviceGroup() );
 
-	// create task chains
+
+	UnifiedRunCommand->SetDescription(TEXT("Unified Run"));
+	UnifiedRunCommand->SetBuildMode(ELauncherProfileBuildModes::DoNotBuild);
+	UnifiedRunCommand->SetCookMode(ELauncherProfileCookModes::DoNotCook);
+	UnifiedRunCommand->SetPackagingMode(ELauncherProfilePackagingModes::DoNotPackage);
+	UnifiedRunCommand->SetDeploymentMode(ELauncherProfileDeploymentModes::DoNotDeploy);
+	UnifiedRunCommand->SetLaunchMode(ELauncherProfileLaunchModes::DefaultRole);
+	UnifiedRunCommand->SetAdditionalUATCommandLine(TEXT(" -skipbuild -skipcook -skipstage -skipdeploy")); // @todo: check if any of these are implied by the items above
+
+	for ( const ILauncherProfileBuildCookRunRef& BuildCookRun : BuildCookRuns)
+	{
+		// unify platforms & targets
+		for (const FString& CookedPlatform : BuildCookRun->GetCookedPlatforms())
+		{
+			UnifiedRunCommand->AddCookedPlatform(CookedPlatform);
+		}
+
+		for (const FString& BuildTarget : BuildCookRun->GetBuildTargets())
+		{
+			UnifiedRunCommand->AddBuildTarget(BuildTarget);
+			UnifiedRunCommand->SetBuildTargetSpecified(true);
+		}
+
+		// unify architectures
+		TArray<FString> ClientArchitectures;
+		for (const FString& ClientArchitecture : BuildCookRun->GetClientArchitectures())
+		{
+			ClientArchitectures.AddUnique(ClientArchitecture);
+		}
+		UnifiedRunCommand->SetClientArchitectures(ClientArchitectures);
+
+		TArray<FString> ServerArchitectures;
+		for (const FString& ServerArchitecture : BuildCookRun->GetServerArchitectures())
+		{
+			ServerArchitectures.AddUnique(ServerArchitecture);
+		}
+		UnifiedRunCommand->SetServerArchitectures(ServerArchitectures);
+
+		TArray<FString> EditorArchitectures;
+		for (const FString& EditorArchitecture : BuildCookRun->GetEditorArchitectures())
+		{
+			EditorArchitectures.AddUnique(EditorArchitecture);
+		}
+		UnifiedRunCommand->SetEditorArchitectures(EditorArchitectures);
+
+		// unify launch command lines
+		UnifiedRunCommand->SetAdditionalCommandLineParameters( CombineCommandLine( UnifiedRunCommand->GetAdditionalCommandLineParameters(), BuildCookRun->GetAdditionalCommandLineParameters() ) );
+		UnifiedRunCommand->SetAdditionalTargetCommandLineParameters( CombineCommandLine( UnifiedRunCommand->GetAdditionalTargetCommandLineParameters(EBuildTargetType::Client), BuildCookRun->GetAdditionalTargetCommandLineParameters(EBuildTargetType::Client) ), EBuildTargetType::Client );
+		UnifiedRunCommand->SetAdditionalTargetCommandLineParameters( CombineCommandLine( UnifiedRunCommand->GetAdditionalTargetCommandLineParameters(EBuildTargetType::Server), BuildCookRun->GetAdditionalTargetCommandLineParameters(EBuildTargetType::Server) ), EBuildTargetType::Server );
+		UnifiedRunCommand->SetAdditionalTargetCommandLineParameters( CombineCommandLine( UnifiedRunCommand->GetAdditionalTargetCommandLineParameters(EBuildTargetType::Game),   BuildCookRun->GetAdditionalTargetCommandLineParameters(EBuildTargetType::Game) ),   EBuildTargetType::Game   );
+		UnifiedRunCommand->SetAdditionalTargetCommandLineParameters( CombineCommandLine( UnifiedRunCommand->GetAdditionalTargetCommandLineParameters(EBuildTargetType::Editor), BuildCookRun->GetAdditionalTargetCommandLineParameters(EBuildTargetType::Editor) ), EBuildTargetType::Editor );
+
+		// unify devices
+		for ( const FString& DeviceId : BuildCookRun->GetDeployedDeviceGroup()->GetDeviceIDs())
+		{
+			UnifiedRunCommand->GetDeployedDeviceGroup()->AddDevice(DeviceId);
+		}
+
+		// first come, first served on the initial map
+		if (UnifiedRunCommand->GetDefaultLaunchRole()->GetInitialMap().IsEmpty())
+		{
+			UnifiedRunCommand->GetDefaultLaunchRole()->SetInitialMap(BuildCookRun->GetDefaultLaunchRole()->GetInitialMap());
+		}
+
+		// ... @todo: any other properties that need unification? ...
+	}
+
+	return UnifiedRunCommand;
+}
+
+TArray<FString> FLauncherWorker::FinalizePlatforms( const TArray<FString>& CookPlatforms, const ILauncherDeviceGroupPtr& DeviceGroup, bool bPreferCookedPlatforms ) const
+{
 	TArray<FString> Platforms;
-	if (InProfile->GetCookMode() == ELauncherProfileCookModes::ByTheBook || InProfile->ShouldBuild())
+	if (bPreferCookedPlatforms)
 	{
-		Platforms = InProfile->GetCookedPlatforms();
+		Platforms = CookPlatforms;
 	}
 
 	// determine deployment platforms
-	ILauncherDeviceGroupPtr DeviceGroup = InProfile->GetDeployedDeviceGroup();
 	FName Variant = NAME_None;
 
 	if (DeviceGroup.IsValid() && Platforms.Num() < 1)
@@ -1450,16 +1609,45 @@ void FLauncherWorker::CreateAndExecuteTasks( const ILauncherProfileRef& InProfil
 	// could not get a platform from the device proxy (suggesting to use the default device) so use the selected cook platforms
 	if (Platforms.Num() == 0)
 	{
-		Platforms = InProfile->GetCookedPlatforms();
+		Platforms = CookPlatforms;
 	}
 
+	return MoveTemp(Platforms);
+}
+
+
+
+/* FLauncherWorker implementation
+ *****************************************************************************/
+
+void FLauncherWorker::CreateAndExecuteTasks( const ILauncherProfileRef& InProfile )
+{
+	FPlatformProcess::CreatePipe(ReadPipe, WritePipe);
+
+	TArray<FString> Platforms;
+	TArray<FString> DeviceIDs;
+	for (const ILauncherProfileBuildCookRunRef& BuildCookRun : InProfile->GetBuildCookRunCommands())
+	{
+		bool bPreferCookedPlatforms = BuildCookRun->GetCookMode() == ELauncherProfileCookModes::ByTheBook || BuildCookRun->ShouldBuild();
+		TArray<FString> FinalPlatforms = FinalizePlatforms(BuildCookRun->GetCookedPlatforms(), BuildCookRun->GetDeployedDeviceGroup(), bPreferCookedPlatforms);
+		for (const FString& Platform : FinalPlatforms)
+		{
+			Platforms.AddUnique(Platform);
+		}
+
+		for (const FString& DeviceID : BuildCookRun->GetDeployedDeviceGroup()->GetDeviceIDs())
+		{
+			DeviceIDs.AddUnique(DeviceID);
+		}
 
 #if !WITH_EDITOR
-	// can't cook by the book in the editor if we are not in the editor...
-	check( InProfile->GetCookMode() != ELauncherProfileCookModes::ByTheBookInEditor );
-	check( InProfile->GetCookMode() != ELauncherProfileCookModes::OnTheFlyInEditor );
+		// can't cook by the book in the editor if we are not in the editor...
+		check( BuildCookRun->GetCookMode() != ELauncherProfileCookModes::ByTheBookInEditor );
+		check( BuildCookRun->GetCookMode() != ELauncherProfileCookModes::OnTheFlyInEditor );
 #endif
+	}
 
+	// create task chains
 	TSharedPtr<FLauncherTask> NextTask;
 	auto AddTask = [this, &NextTask](TSharedPtr<FLauncherTask> NewTask)
 	{
@@ -1513,36 +1701,51 @@ void FLauncherWorker::CreateAndExecuteTasks( const ILauncherProfileRef& InProfil
 	int UATCommandIndex = 0;
 
 
+	FCommandAdjustments CommandAdjustments;
 	TArray<FCommandDesc> Commands;
 	FString StartString;
+	FString UATCommandString;
 
-	FString PreCommand;
-	if (Profile->ShouldUpdateDeviceFlash() && DeviceGroup->GetDeviceIDs().Num() >= 0)
+	// special case for Turnkey
+	if (Profile->ShouldUpdateDeviceFlash() && DeviceIDs.Num() > 0)
 	{
-		PreCommand = FString::Printf(TEXT("Turnkey -command=VerifySdk -type=Flash -device=%s -UpdateIfNeeded -utf8output -WaitForUATMutex %s"), *FString::Join(DeviceGroup->GetDeviceIDs(), TEXT("+")), *ITurnkeyIOModule::Get().GetUATParams());
-		PreCommand += TEXT(" ");
+		UATCommandString += FString::Printf(TEXT("Turnkey -command=VerifySdk -type=Flash -device=%s -UpdateIfNeeded -WaitForUATMutex %s "), *FString::Join(DeviceIDs, TEXT("+")), *ITurnkeyIOModule::Get().GetUATParams());
 	}
 
-	// add pre-BuildCookRun commands
-	while ( UATCommandIndex < UATCommands.Num() && UATCommands[UATCommandIndex]->GetOrder() < 0)
+	
+	ILauncherProfileUATCommandPtr UnifiedRunUATCommand;
+	if (InProfile->GetAutomatedTests().Num() > 0)
 	{
-		ILauncherProfileUATCommandRef UATCommand = UATCommands[UATCommandIndex++];
-		PreCommand += CreateUATCommand(InProfile, UATCommand, Platforms, Commands, StartString);
-		PreCommand += TEXT(" ");
+		// disable run in BuildCookRun commands if there are automated tests as they will take care of the run
+		CommandAdjustments.bDisableRun = true;
+	}
+	else
+	{
+		// disable run in BuidCookRun commands if there is multiple BuildCookRun commands that want to run - in this case, they're gathered together and run in a separate unified run command at the end
+		UnifiedRunUATCommand = MakeTemporaryUnifiedRunCommand(InProfile);
+		if (UnifiedRunUATCommand.IsValid())
+		{
+			UATCommands.Add(UnifiedRunUATCommand.ToSharedRef());
+			CommandAdjustments.bDisableRun = true;
+		}
 	}
 
-	// main BuildCookRun command
-	FString MainUATCommand = CreateUATCommandForBuildCookRun(InProfile, Platforms, Commands, StartString);
-	
-	// add post-BuildCookRun commands
-	FString PostCommand;
-	while ( UATCommandIndex < UATCommands.Num())
+
+	// build all UAT commands
+	for (const ILauncherProfileUATCommandRef& UATCommand : UATCommands)
 	{
-		ILauncherProfileUATCommandRef UATCommand = UATCommands[UATCommandIndex++];
-		PostCommand += CreateUATCommand(InProfile, UATCommand, Platforms, Commands, StartString);
-		PostCommand += TEXT(" ");
+		FCommandAdjustments Adjustments = CommandAdjustments;
+		if (UATCommand == UnifiedRunUATCommand)
+		{
+			Adjustments.bDisableRun = false;
+			Adjustments.bPreferCookedPlatforms = true;
+		}
+
+		UATCommandString += CreateUATCommand(UATCommand, InProfile, Platforms, Commands, StartString, Adjustments);
+		UATCommandString += TEXT(" ");
 	}
 
+
 	// wait for completion of automated tests
 	if (Profile->GetAutomatedTests().Num() > 0)
 	{
@@ -1566,7 +1769,7 @@ void FLauncherWorker::CreateAndExecuteTasks( const ILauncherProfileRef& InProfil
 	}
 
 
-	TSharedPtr<FLauncherTask> LaunchTask = MakeShareable(new FLauncherUATTask(MainUATCommand, TEXT("Launch Task"), TEXT("Launching UAT..."), ReadPipe, WritePipe, InProfile->GetEditorExe(), ProcHandle, this, StartString, PreCommand.TrimStartAndEnd(), PostCommand.TrimStartAndEnd()));
+	TSharedPtr<FLauncherTask> LaunchTask = MakeShareable(new FLauncherUATTask(UATCommandString, TEXT("Launch Task"), TEXT("Launching UAT..."), ReadPipe, WritePipe, InProfile->GetEditorExe(), ProcHandle, this, StartString));
 	AddTask(LaunchTask);
 	for (int32 Index = 0; Index < Commands.Num(); ++Index)
 	{
diff --git a/Engine/Source/Developer/LauncherServices/Private/Launcher/LauncherWorker.h b/Engine/Source/Developer/LauncherServices/Private/Launcher/LauncherWorker.h
--- a/Engine/Source/Developer/LauncherServices/Private/Launcher/LauncherWorker.h
+++ b/Engine/Source/Developer/LauncherServices/Private/Launcher/LauncherWorker.h
@@ -116,11 +116,19 @@ protected:
 	void OnTaskStarted(const FString& TaskName);
 	void OnTaskCompleted(const FString& TaskName);
 
-	FString CreateUATCommandForBuildCookRun( const ILauncherProfileRef& InProfile, const TArray<FString>& InPlatforms, TArray<FCommandDesc>& OutCommands, FString& CommandStart, bool bForTurnkeyCustomBuild = false );
+	struct FCommandAdjustments
+	{
+		bool bDisableRun = false;
+		bool bPreferCookedPlatforms = false;
+	};
+	FString CreateUATCommandForBuildCookRun( const ILauncherProfileBuildCookRunRef& InBuildCookRun, const ILauncherProfileRef& InProfile, TArray<FCommandDesc>& OutCommands, FString& OutCommandStart, const FCommandAdjustments& Adjustments, bool bForTurnkeyCustomBuild = false );
+
+	FString CreateUATCommand( const ILauncherProfileUATCommandRef& InUATCommand, const ILauncherProfileRef& InProfile, const TArray<FString>& InPlatforms, TArray<FCommandDesc>& OutCommands, FString& OutCommandStart, const FCommandAdjustments& Adjustments );
 
-	FString CreateUATCommand( const ILauncherProfileRef& InProfile, ILauncherProfileUATCommandRef& InUATCommand, const TArray<FString>& InPlatforms, TArray<FCommandDesc>& OutCommands, FString& OutCommandStart );
+	FString CreateUATCommandForAutomatedTest( const ILauncherProfileAutomatedTestRef& InAutomatedTest, const ILauncherProfileRef& InProfile, const TArray<FString>& InPlatforms, TArray<FCommandDesc>& OutCommands, FString& OutCommandStart, const FCommandAdjustments& Adjustments );
 
-	FString CreateUATCommandForAutomatedTest( const ILauncherProfileAutomatedTestRef& InAutomatedTest, const ILauncherProfileRef& InProfile, const TArray<FString>& InPlatforms, TArray<FCommandDesc>& OutCommands, FString& OutCommandStart );
+	ILauncherProfileUATCommandPtr MakeTemporaryUnifiedRunCommand(const ILauncherProfileRef& InProfile);
+	TArray<FString> FinalizePlatforms( const TArray<FString>& CookPlatforms, const ILauncherDeviceGroupPtr& DeviceGroup, bool bPreferCookedPlatforms ) const;
 
 	//start listening to the device discovery routine
 	void EnableDeviceDiscoveryListener();
diff --git a/Engine/Source/Developer/LauncherServices/Private/Profiles/LauncherProfile.cpp b/Engine/Source/Developer/LauncherServices/Private/Profiles/LauncherProfile.cpp
--- a/Engine/Source/Developer/LauncherServices/Private/Profiles/LauncherProfile.cpp
+++ b/Engine/Source/Developer/LauncherServices/Private/Profiles/LauncherProfile.cpp
@@ -3,12 +3,72 @@
 #include "LauncherProfile.h"
 #include "Profiles/LauncherProfileUATCommand.h"
 #include "Profiles/LauncherProfileAutomatedTest.h"
+#include "Profiles/LauncherProfileBuildCookRun.h"
 
 #define LOCTEXT_NAMESPACE "SProjectLauncherValidation"
 
+DEFINE_LOG_CATEGORY(LogLauncherProfile);
 
 const TCHAR* FLauncherProfileUATCommand::TypeName = TEXT("UATCommand");
 const TCHAR* FLauncherProfileAutomatedTest::TypeName = TEXT("AutomatedTest");
+const TCHAR* FLauncherProfileBuildCookRun::TypeName = TEXT("BuildCookRun");
+
+
+ILauncherProfileUATCommandRef FLauncherProfile::CreateUATCommandInternal(const TCHAR* InternalName)
+{
+	ILauncherProfileUATCommandRef UATCommand = MakeShared<FLauncherProfileUATCommand>(InternalName);
+	return UATCommand;
+}
+
+ILauncherProfileAutomatedTestRef FLauncherProfile::CreateAutomatedTestInternal(const TCHAR* InternalName)
+{
+	ILauncherProfileAutomatedTestRef AutomatedTest = MakeShared<FLauncherProfileAutomatedTest>(InternalName);
+	return AutomatedTest;
+}
+
+ILauncherProfileBuildCookRunRef FLauncherProfile::CreateBuildCookRunInternal(const TCHAR* InternalName)
+{
+	ILauncherProfileBuildCookRunRef BuildCookRun = MakeShared<FLauncherProfileBuildCookRun>(LauncherProfileManager, this, InternalName);
+	return BuildCookRun;
+
+}
+
+ILauncherProfileUATCommandPtr FLauncherProfile::CreateUATCommandByTypeNameInternal( const TCHAR* TypeName, const TCHAR* InternalName )
+{
+	if (FCString::Stricmp(TypeName, FLauncherProfileUATCommand::TypeName) == 0)
+	{
+		return CreateUATCommandInternal(InternalName);
+	}
+	else if (FCString::Stricmp(TypeName, FLauncherProfileAutomatedTest::TypeName) == 0)
+	{
+		return CreateAutomatedTestInternal(InternalName);
+	}
+	else if (FCString::Stricmp(TypeName, FLauncherProfileBuildCookRun::TypeName) == 0)
+	{
+		return CreateBuildCookRunInternal(InternalName);
+	}
+	else
+	{
+		return nullptr;
+	}
+}
+
+
+ILauncherProfileBuildCookRunPtr FLauncherProfile::LoadBuildCookRunInternal(const FJsonObject& Object, int32 Version)
+{
+	TSharedRef<FLauncherProfileBuildCookRun> BuildCookRun = MakeShared<FLauncherProfileBuildCookRun>(LauncherProfileManager, this, TEXT("OriginalBuildCookRun"));
+	if (!BuildCookRun->Load(Object, Version))
+	{
+		return nullptr;
+	}
+
+	// should always have a device group
+	ILauncherDeviceGroupRef DeployDeviceGroup = LauncherProfileManager->AddNewDeviceGroup();
+	BuildCookRun->SetDeployedDeviceGroup(DeployDeviceGroup);
+
+	return BuildCookRun;
+}
+
 
 
 
@@ -130,7 +190,16 @@ int32 FLauncherProfile::GetZenSnapshot()
 	FString ProjectName = GetProjectName().ToLower();
 	if (PerProjectBuilds.Contains(ProjectName))
 	{
-		TArray<FString> Platforms = GetCookedPlatforms();
+		TArray<FString> Platforms;
+		for (const ILauncherProfileBuildCookRunRef& BuildCookRun : GetBuildCookRunCommands())
+		{
+			TArray<FString> ThisPlatforms = BuildCookRun->GetCookedPlatforms();
+			for (const FString& Platform : ThisPlatforms)
+			{
+				Platforms.AddUnique(Platform);
+			}
+		}
+
 		for (const FString& Platform : Platforms)
 		{
 			FString LowerPlatform = Platform.ToLower();
@@ -159,5 +228,331 @@ TSharedPtr<UE::Zen::Build::FBuildListRetriever> FLauncherProfile::GetBuildListRe
 	return BuildListRetriever;
 }
 
+
+
+
+bool FLauncherProfile::Serialize( FArchive& Archive )
+{
+	// use our friend access to help with deprected serialization
+	TSharedPtr<FLauncherProfileBuildCookRun> PrivateInner = StaticCastSharedPtr<FLauncherProfileBuildCookRun>(FirstBuildCookRun);
+	if (!PrivateInner.IsValid())
+	{
+		return false;
+	}
+
+	int32 Version = LAUNCHERSERVICES_FINAL;
+
+	Archive	<< Version;
+
+	if (Version < LAUNCHERSERVICES_MINPROFILEVERSION)
+	{
+		return false;
+	}
+
+	if (Version > LAUNCHERSERVICES_FINAL)
+	{
+		return false;
+	}
+
+	if (Archive.IsSaving())
+	{
+		if (PrivateInner->DeployedDeviceGroup.IsValid())
+		{
+			PrivateInner->DeployedDeviceGroupId = PrivateInner->DeployedDeviceGroup->GetId();
+		}
+		else
+		{
+			PrivateInner->DeployedDeviceGroupId = FGuid();
+		}
+	}
+
+	// IMPORTANT: make sure to bump LAUNCHERSERVICES_PROFILEVERSION when modifying this!
+	bool BuildGame = false;
+	Archive << Id
+			<< Name
+			<< Description
+			<< PrivateInner->BuildConfiguration
+			<< ProjectSpecified
+			<< ShareableProjectPath
+			<< PrivateInner->CookConfiguration
+			<< PrivateInner->CookIncremental
+			<< PrivateInner->CookOptions
+			<< PrivateInner->CookMode
+			<< PrivateInner->CookUnversioned
+			<< PrivateInner->CookedCultures
+			<< PrivateInner->CookedMaps
+			<< PrivateInner->CookedPlatforms
+			<< PrivateInner->DeployStreamingServer
+			<< PrivateInner->DeployWithUnrealPak
+			<< PrivateInner->DeployedDeviceGroupId
+			<< PrivateInner->DeploymentMode
+			<< PrivateInner->HideFileServerWindow
+			<< PrivateInner->LaunchMode
+			<< PrivateInner->PackagingMode
+			<< PrivateInner->PackageDir
+			<< BuildGame
+            << PrivateInner->ForceClose
+            << PrivateInner->Timeout;
+
+	if (Version >= LAUNCHERSERVICES_SHAREABLEPROJECTPATHS)
+	{
+		FullProjectPath = FPaths::ConvertRelativePathToFull(FPaths::RootDir(), ShareableProjectPath);
+	}
+
+	FString DeployPlatformString = PrivateInner->DefaultDeployPlatform.ToString();
+	if (Version >= LAUNCHERSERVICES_FIXCOMPRESSIONSERIALIZE)
+	{
+		Archive << PrivateInner->Compressed;
+	}
+	if ( Version>= LAUNCHERSERVICES_ADDEDENCRYPTINIFILES)
+	{
+		Archive << PrivateInner->EncryptIniFiles;
+		Archive << PrivateInner->ForDistribution;
+	}
+	if (Version >= LAUNCHERSERVICES_ADDEDDEFAULTDEPLOYPLATFORM)
+	{
+		Archive << DeployPlatformString;
+	}		
+	if (Version >= LAUNCHERSERVICES_ADDEDNUMCOOKERSTOSPAWN && Version < LAUNCHERSERVICES_REMOVEDNUMCOOKERSTOSPAWN)
+	{
+		int32 OldNumCookersToSpawn;
+		Archive << OldNumCookersToSpawn;
+	}
+	if (Version >= LAUNCHERSERVICES_ADDEDSKIPCOOKINGEDITORCONTENT)
+	{
+		Archive << PrivateInner->bSkipCookingEditorContent;
+	}
+	if (Version >= LAUNCHERSERVICES_ADDEDINCREMENTALDEPLOYVERSION)
+	{
+		Archive << PrivateInner->DeployIncremental;
+	}
+	if ( Version >= LAUNCHERSERVICES_REMOVEDPATCHSOURCECONTENTPATH )
+	{
+		Archive << PrivateInner->GeneratePatch;
+	}
+	if ( Version >= LAUNCHERSERVICES_ADDEDMULTILEVELPATCHING )
+	{
+		Archive << PrivateInner->AddPatchLevel;
+		Archive << StageBaseReleasePaks;
+	}
+	else if ( Version >= LAUNCHERSERVICES_ADDEDPATCHSOURCECONTENTPATH)
+	{
+		FString Temp;
+		Archive << Temp;
+		Archive << PrivateInner->GeneratePatch;
+	}
+	if (Version >= LAUNCHERSERVICES_ADDEDDLCINCLUDEENGINECONTENT)
+	{
+		Archive << PrivateInner->DLCIncludeEngineContent;
+	}
+		
+	if ( Version >= LAUNCHERSERVICES_ADDEDRELEASEVERSION )
+	{
+		Archive << PrivateInner->CreateReleaseVersion;
+		Archive << PrivateInner->CreateReleaseVersionName;
+		Archive << PrivateInner->BasedOnReleaseVersionName;
+			
+		Archive << PrivateInner->CreateDLC;
+		Archive << PrivateInner->DLCName;
+	}
+	if (Version >= LAUNCHERSERVICES_ADDEDGENERATECHUNKS)
+	{
+		Archive << PrivateInner->bGenerateChunks;
+		Archive << PrivateInner->bGenerateHttpChunkData;
+		Archive << PrivateInner->HttpChunkDataDirectory;
+		Archive << PrivateInner->HttpChunkDataReleaseName;
+	}
+	if (Version >= LAUNCHERSERVICES_ADDARCHIVE)
+	{
+		Archive << PrivateInner->bArchive;
+		Archive << PrivateInner->ArchiveDir;
+	}
+	if (Version >= LAUNCHERSERVICES_ADDEDADDITIONALCOMMANDLINE)
+	{
+		Archive << PrivateInner->AdditionalCommandLineParameters;
+	}
+	if (Version >= LAUNCHERSERVICES_ADDEDINCLUDEPREREQUISITES)
+	{
+		Archive << PrivateInner->IncludePrerequisites;
+	}
+
+	if (Version >= LAUNCHERSERVICES_ADDEDBUILDMODE)
+	{
+		Archive << PrivateInner->BuildMode;
+	}
+	else if(Archive.IsLoading())
+	{
+		PrivateInner->BuildMode = BuildGame ? ELauncherProfileBuildModes::Build : ELauncherProfileBuildModes::DoNotBuild;
+	}
+
+	if (Version >= LAUNCHERSERVICES_ADDEDUSEIOSTORE)
+	{
+		Archive << PrivateInner->bUseIoStore;
+	}
+
+	if (Version >= LAUNCHERSERVICES_ADDEDMAKEBINARYCONFIG)
+	{
+		Archive << PrivateInner->bMakeBinaryConfig;
+	}
+
+	if (Version >= LAUNCHERSERVICES_ADDEDREFERENCECONTAINERS)
+	{
+		Archive << PrivateInner->ReferenceContainerGlobalFileName;
+		Archive << PrivateInner->ReferenceContainerCryptoKeysFileName;
+	}
+	if (Version >= LAUNCHERSERVICES_ADDEDORIGINALRELEASEVERSION)
+	{
+		Archive << PrivateInner->OriginalReleaseVersionName;
+	}
+
+	FString BuildTargetName;
+	if (Version >= LAUNCHERSERVICES_ADDBUILDTARGETNAME)
+	{
+		Archive << PrivateInner->BuildTargetSpecified;
+		Archive << BuildTargetName;
+	}
+
+	if (Version >= LAUNCHERSERVICES_ADDEDRETAINSTAGEDDIRECTORY && Version < LAUNCHERSERVICES_REMOVEDRETAINSTAGEDDIRECTORY)
+	{
+		bool bRetainStagedDirectory = false;
+		Archive << bRetainStagedDirectory;
+	}
+
+	if (Version >= LAUNCHERSERVICES_ADDZENOPTIONS)
+	{
+		Archive << PrivateInner->bUseZenStore;
+		Archive << PrivateInner->bImportingZenSnapshot;
+		Archive << PrivateInner->bUseZenPakStreaming;
+		Archive << PrivateInner->ZenPakStreamingPath;
+	}
+
+	if (Version >= LAUNCHERSERVICES_ADDBUILDARCHITECTURE)
+	{
+		Archive << PrivateInner->ClientArchitectures;
+		Archive << PrivateInner->ServerArchitectures;
+		Archive << PrivateInner->EditorArchitectures;
+	}
+
+	if (Version >= LAUNCHERSERVICES_ADDCUSTOMPROPERTIES)
+	{
+		Archive << CustomStringProperties;
+		Archive << CustomBoolProperties;
+	}
+
+	if (Version >= LAUNCHERSERVICES_INCREMENTALCOOKMODE && Version < LAUNCHERSERVICES_REMOVEDINCREMENTALCOOKMODE)
+	{
+		TEnumAsByte<ELauncherProfileIncrementalCookMode::Type> IncrementalCookMode = (PrivateInner->CookIncremental ? ELauncherProfileIncrementalCookMode::ModifiedOnly : ELauncherProfileIncrementalCookMode::None);
+		Archive << IncrementalCookMode;
+	}
+
+	if (Version >= LAUNCHERSERVICES_AUTOMATEDTESTS)
+	{
+		TArray<TSharedPtr<ILauncherProfileAutomatedTest>> AutomatedTests;
+		if (Archive.IsSaving())
+		{
+			for (const ILauncherProfileUATCommandRef& UATCommand : UATCommands)
+			{
+				ILauncherProfileAutomatedTestPtr AutomatedTest = UATCommand->AsAutomatedTest();
+				if (AutomatedTest)
+				{
+					AutomatedTests.Add(AutomatedTest);
+				}
+			}
+		}
+
+		// serialize automated tests
+		int NumAutomatedTests = AutomatedTests.Num();
+		Archive << NumAutomatedTests;
+		for( int32 TestIndex = 0; TestIndex < NumAutomatedTests; TestIndex++)
+		{
+PRAGMA_DISABLE_DEPRECATION_WARNINGS
+			if (Archive.IsLoading())
+			{
+				ILauncherProfileAutomatedTestRef AutomatedTest = CreateAutomatedTestInternal();
+				AutomatedTest->Serialize(Archive);
+				AutomatedTests.Add(AutomatedTest);
+			}
+			else
+			{
+				AutomatedTests[TestIndex]->Serialize(Archive);
+			}
+		}
+PRAGMA_ENABLE_DEPRECATION_WARNINGS
+
+		if (Archive.IsLoading())
+		{
+			for (TSharedPtr<ILauncherProfileAutomatedTest> AutomatedTest : AutomatedTests)
+			{
+				UATCommands.Add(AutomatedTest.ToSharedRef());
+			}
+		}
+
+		Archive << AutomatedTestBuildPath;
+		Archive << bWantAutomatedTestBuild;
+	}
+
+	if (Version >= LAUNCHERSERVICES_ADDMULTIPLEBUILDTARGETS)
+	{
+		Archive << PrivateInner->BuildTargets;
+		if (Archive.IsLoading() && PrivateInner->BuildTargets.Num() == 0 && !BuildTargetName.IsEmpty())
+		{
+			PrivateInner->BuildTargets.Add(BuildTargetName);
+		}
+
+		Archive << PrivateInner->AdditionalTargetCommandLineParameters;
+	}
+
+PRAGMA_DISABLE_DEPRECATION_WARNINGS
+	PrivateInner->DefaultLaunchRole->Serialize(Archive);
+PRAGMA_ENABLE_DEPRECATION_WARNINGS
+
+	// serialize launch roles
+	if (Archive.IsLoading())
+	{
+		PrivateInner->DeployedDeviceGroup.Reset();
+		PrivateInner->LaunchRoles.Reset();
+	}
+
+	int32 NumLaunchRoles = PrivateInner->LaunchRoles.Num();
+
+	Archive << NumLaunchRoles;
+
+	for (int32 RoleIndex = 0; RoleIndex < NumLaunchRoles; ++RoleIndex)
+	{
+PRAGMA_DISABLE_DEPRECATION_WARNINGS
+		if (Archive.IsLoading())
+		{
+			PrivateInner->LaunchRoles.Add(MakeShareable(new FLauncherProfileLaunchRole(Archive)));				
+		}
+		else
+		{
+			PrivateInner->LaunchRoles[RoleIndex]->Serialize(Archive);
+		}
+PRAGMA_ENABLE_DEPRECATION_WARNINGS
+	}
+
+	if (Archive.IsLoading())
+	{
+		PrivateInner->DefaultDeployPlatform = FName(*DeployPlatformString);
+	}
+
+	if (PrivateInner->DefaultDeployPlatform != NAME_None)
+	{
+		PrivateInner->SetDefaultDeployPlatform(PrivateInner->DefaultDeployPlatform);
+	}
+
+	if (Archive.IsLoading())
+	{
+		RefreshValidBuildTargets();
+	}
+	Validate();
+
+	return true;
+}
+
+
+
+
+
 #undef LOCTEXT_NAMESPACE
 
