Launcher Profile support for multiple build cook run commands per-profile.

diff --git a/Engine/Source/Developer/LauncherServices/Private/Profiles/LauncherProfile.h b/Engine/Source/Developer/LauncherServices/Private/Profiles/LauncherProfile.h
--- a/Engine/Source/Developer/LauncherServices/Private/Profiles/LauncherProfile.h
+++ b/Engine/Source/Developer/LauncherServices/Private/Profiles/LauncherProfile.h
@@ -18,86 +18,22 @@
 #include "ITargetDeviceProxy.h"
 #include "ITargetDeviceServicesModule.h"
 #include "GameProjectHelper.h"
+#include "Profiles/LauncherProfileVersion.h"
 #include "Profiles/LauncherProfileLaunchRole.h"
-#include "Profiles/LauncherProfileUATCommand.h"
-#include "Profiles/LauncherProfileAutomatedTest.h"
 #include "PlatformInfo.h"
 #include "TargetReceipt.h"
 #include "DesktopPlatformModule.h"
 #include "GameProjectHelper.h"
 #include "Experimental/BuildListRetriever.h"
 
-DEFINE_LOG_CATEGORY_STATIC(LogLauncherProfile, Log, All);
-
 class Error;
 
-enum ELauncherVersion
-{
-	LAUNCHERSERVICES_MINPROFILEVERSION=10,
-	LAUNCHERSERVICES_ADDEDINCREMENTALDEPLOYVERSION=11,
-	LAUNCHERSERVICES_ADDEDPATCHSOURCECONTENTPATH=12,
-	LAUNCHERSERVICES_ADDEDRELEASEVERSION=13,
-	LAUNCHERSERVICES_REMOVEDPATCHSOURCECONTENTPATH=14,
-	LAUNCHERSERVICES_ADDEDDLCINCLUDEENGINECONTENT=15,
-	LAUNCHERSERVICES_ADDEDGENERATECHUNKS=16,
-	LAUNCHERSERVICES_ADDEDNUMCOOKERSTOSPAWN=17,
-	LAUNCHERSERVICES_ADDEDSKIPCOOKINGEDITORCONTENT=18,
-	LAUNCHERSERVICES_ADDEDDEFAULTDEPLOYPLATFORM=19,
-	LAUNCHERSERVICES_FIXCOMPRESSIONSERIALIZE=20,
-	LAUNCHERSERVICES_SHAREABLEPROJECTPATHS = 21,
-	LAUNCHERSERVICES_FILEFORMATCHANGE = 22,
-	LAUNCHERSERVICES_ADDARCHIVE = 23,
-	LAUNCHERSERVICES_ADDEDENCRYPTINIFILES = 24,
-	LAUNCHERSERVICES_ADDEDMULTILEVELPATCHING = 25,
-	LAUNCHERSERVICES_ADDEDADDITIONALCOMMANDLINE = 26,
-	LAUNCHERSERVICES_ADDEDINCLUDEPREREQUISITES = 27,
-	LAUNCHERSERVICES_ADDEDBUILDMODE = 28,
-	LAUNCHERSERVICES_ADDEDUSEIOSTORE = 29,
-	LAUNCHERSERVICES_ADDEDMAKEBINARYCONFIG = 30,
-	LAUNCHERSERVICES_ADDEDREFERENCECONTAINERS = 31,
-	LAUNCHERSERVICES_REMOVEDNUMCOOKERSTOSPAWN = 32,
-	LAUNCHERSERVICES_ADDEDORIGINALRELEASEVERSION = 33,
-	LAUNCHERSERVICES_ADDBUILDTARGETNAME = 34,
-	LAUNCHERSERVICES_ADDEDRETAINSTAGEDDIRECTORY = 35,
-	LAUNCHERSERVICES_REMOVEDRETAINSTAGEDDIRECTORY = 36,
-	LAUNCHERSERVICES_ADDZENOPTIONS = 37,
-	LAUNCHERSERVICES_ADDBUILDARCHITECTURE = 38,
-	LAUNCHERSERVICES_ADDCUSTOMPROPERTIES = 39,
-	LAUNCHERSERVICES_INCREMENTALCOOKMODE = 40,
-	LAUNCHERSERVICES_AUTOMATEDTESTS = 41,
-	LAUNCHERSERVICES_FASTITERATE = 42,
-	LAUNCHERSERVICES_REMOVEDINCREMENTALCOOKMODE = 43,
-	LAUNCHERSERVICES_ADDMULTIPLEBUILDTARGETS = 44,
-	LAUNCHERSERVICES_ADDUATCOMMANDS = 45,
-	//ADD NEW STUFF HERE
-
-
-	LAUNCHERSERVICES_FINALPLUSONE,
-	LAUNCHERSERVICES_FINAL = LAUNCHERSERVICES_FINALPLUSONE-1,
-};
-
 enum ESimpleLauncherVersion
 {
 	LAUNCHERSERVICES_SIMPLEPROFILEVERSION=1,
 	LAUNCHERSERVICES_SIMPLEFILEFORMATCHANGE = 2,
 };
 
-LAUNCHERSERVICES_API bool HasPromotedTarget(const TCHAR* BaseDir, const TCHAR* TargetName, const TCHAR* Platform, EBuildConfiguration Configuration, const TCHAR* Architecture);
-
-inline bool TryGetDefaultTargetName(const FString& ProjectFile, EBuildTargetType TargetType, FString& OutTargetName)
-{
-	const TArray<FTargetInfo>& Targets = FDesktopPlatformModule::Get()->GetTargetsForProject(ProjectFile);
-	for (const FTargetInfo& Target : Targets)
-	{
-		if (Target.Type == TargetType && Target.DefaultTarget.Get(true))
-		{
-			OutTargetName = Target.Name;
-			return true;
-		}
-	}
-	return false;
-}
-
 /**
 * Implements a simple profile which controls the desired output of the Launcher for simple
 */
@@ -263,7 +199,6 @@ public:
 	 */
 	FLauncherProfile(ILauncherProfileManagerRef ProfileManager)
 		: LauncherProfileManager(ProfileManager)
-		, DefaultLaunchRole(MakeShareable(new FLauncherProfileLaunchRole()))
 	{ 
 		SetDefaults();
 		ProjectChangedDelegate.AddRaw(this, &FLauncherProfile::OnSelectedProjectChanged);
@@ -276,7 +211,6 @@ public:
 	 */
 	FLauncherProfile(ILauncherProfileManagerRef ProfileManager, FGuid InId, const FString& InProfileName)
 		: LauncherProfileManager(ProfileManager)
-		, DefaultLaunchRole(MakeShareable(new FLauncherProfileLaunchRole()))
 		, Id(InId)
 		, Name(InProfileName)
 	{
@@ -289,146 +223,102 @@ public:
 	 */
 	~FLauncherProfile( ) 
 	{
-		if (DeployedDeviceGroup.IsValid())
-		{
-			DeployedDeviceGroup->OnDeviceAdded().Remove(OnLauncherDeviceGroupDeviceAddedDelegateHandle);
-			DeployedDeviceGroup->OnDeviceRemoved().Remove(OnLauncherDeviceGroupDeviceRemoveDelegateHandle);
-		}
 	}
 
-	/**
-	 * Gets the identifier of the device group to deploy to.
-	 *
-	 * This method is used internally by the profile manager to read the device group identifier after
-	 * loading this profile from a file. The profile manager will use this identifier to locate the
-	 * actual device group to deploy to.
-	 *
-	 * @return The device group identifier, or an invalid GUID if no group was set or deployment is disabled.
-	 */
-	const FGuid& GetDeployedDeviceGroupId( ) const
+
+	struct FValidationData
 	{
-		return DeployedDeviceGroupId;
-	}
+		void Reset()
+		{
+			Errors.Reset();
+			CustomErrors.Reset();
+			CustomWarnings.Reset();
+			InvalidPlatform.Reset();
+		}
+
+		// Holds the collection of validation errors.
+		TArray<ELauncherProfileValidationErrors::Type> Errors;
+		FString InvalidPlatform;
+
+		// Holds the collection of custom errors & warnings.
+		TMap<FString,FText> CustomErrors;
+		TMap<FString,FText> CustomWarnings;
+	};
+
+
 
 	//~ Begin ILauncherProfile Interface
 
 	virtual void AddCookedCulture( const FString& CultureName ) override
 	{
-		CookedCultures.AddUnique(CultureName);
-
-		Validate();
+		check(FirstBuildCookRun.IsValid());
+		FirstBuildCookRun->AddCookedCulture(CultureName);
 	}
 
 	virtual void AddCookedMap( const FString& MapName ) override
 	{
-		CookedMaps.AddUnique(MapName);
-
-		Validate();
+		check(FirstBuildCookRun.IsValid());
+		FirstBuildCookRun->AddCookedMap(MapName);
 	}
 
 	virtual void AddCookedPlatform( const FString& PlatformName ) override
 	{
-		CookedPlatforms.AddUnique(PlatformName);
-
+		check(FirstBuildCookRun.IsValid());
+		FirstBuildCookRun->AddCookedPlatform(PlatformName);
 		RefreshValidBuildTargets();
 		Validate();
 	}
 
 	virtual void SetDefaultDeployPlatform(const FName PlatformName) override
 	{
-		
-		DefaultDeployPlatform = PlatformName;	
-
-		if (DeployedDeviceGroup.IsValid())
-		{
-			DeployedDeviceGroup->RemoveAllDevices();
-
-			if (DefaultDeployPlatform != NAME_None)
-			{
-				TArray<TSharedPtr<ITargetDeviceProxy>> PlatformDeviceProxies;
-				ITargetDeviceServicesModule& TargetDeviceServicesModule = FModuleManager::LoadModuleChecked<ITargetDeviceServicesModule>("TargetDeviceServices");
-				const TSharedRef<ITargetDeviceProxyManager>& InDeviceProxyManager = TargetDeviceServicesModule.GetDeviceProxyManager();
-
-				InDeviceProxyManager->GetProxies(NAME_None, true, PlatformDeviceProxies);
-
-				TSharedPtr<ITargetDeviceProxy> DefaultPlatformDevice;
-				for (int32 ProxyIndex = 0; ProxyIndex < PlatformDeviceProxies.Num(); ++ProxyIndex)
-				{
-					TSharedPtr<ITargetDeviceProxy> DeviceProxy = PlatformDeviceProxies[ProxyIndex];
-
-					if (DeviceProxy->GetVanillaPlatformId(NAME_None) == DefaultDeployPlatform)
-					{
-						DefaultPlatformDevice = DeviceProxy;
-						break;
-					}
-				}
-
-				if (DefaultPlatformDevice.IsValid())
-				{
-					DeployedDeviceGroup->AddDevice(DefaultPlatformDevice->GetTargetDeviceId(NAME_None));
-				}
-			}
-		}
-
-		Validate();
+		check(FirstBuildCookRun.IsValid());
+		FirstBuildCookRun->SetDefaultDeployPlatform(PlatformName);
 	}
 
 	virtual void ClearCookedCultures( ) override
 	{
-		if (CookedCultures.Num() > 0)
-		{
-			CookedCultures.Reset();
-
-			Validate();
-		}
+		check(FirstBuildCookRun.IsValid());
+		FirstBuildCookRun->ClearCookedCultures();
 	}
 
 	virtual void ClearCookedMaps( ) override
 	{
-		if (CookedMaps.Num() > 0)
-		{
-			CookedMaps.Reset();
-
-			Validate();
-		}
+		check(FirstBuildCookRun.IsValid());
+		FirstBuildCookRun->ClearCookedMaps();
 	}
 
 	virtual void ClearCookedPlatforms( ) override
 	{
-		if (CookedPlatforms.Num() > 0)
-		{
-			CookedPlatforms.Reset();
-
-			RefreshValidBuildTargets();
-			Validate();
-		}
+		check(FirstBuildCookRun.IsValid());
+		FirstBuildCookRun->ClearCookedPlatforms();
+		RefreshValidBuildTargets();
+		Validate();
 	}
 
 	virtual ILauncherProfileLaunchRolePtr CreateLaunchRole( ) override
 	{
-		ILauncherProfileLaunchRolePtr Role = MakeShareable(new FLauncherProfileLaunchRole());
-			
-		LaunchRoles.Add(Role);
-
-		Validate();
-
-		return Role;
+		check(FirstBuildCookRun.IsValid());
+		return FirstBuildCookRun->CreateLaunchRole();
 	}
 
 	virtual EBuildConfiguration GetBuildConfiguration( ) const override
 	{
-		return BuildConfiguration;
+		check(FirstBuildCookRun.IsValid());
+		return FirstBuildCookRun->GetBuildConfiguration();;
 	}
 
 	virtual bool HasBuildTargetSpecified() const override
 	{
-		return BuildTargetSpecified;
+		check(FirstBuildCookRun.IsValid());
+		return FirstBuildCookRun->HasBuildTargetSpecified();
 	}
 
 	virtual FString GetBuildTarget() const override
 	{
-		if (BuildTargetSpecified)
+		if (HasBuildTargetSpecified())
 		{
+			check(FirstBuildCookRun.IsValid());
+			TArray<FString> BuildTargets = FirstBuildCookRun->GetBuildTargets();
 			UE_CLOG(BuildTargets.Num() > 0, LogTemp, Warning, TEXT("GetBuildTargets is only returning the first build target - legacy project launcher?"));
 
 			return (BuildTargets.Num() > 0) ? BuildTargets[0] : FString();
@@ -441,68 +331,74 @@ public:
 
 	virtual EBuildConfiguration GetCookConfiguration( ) const override
 	{
-		return CookConfiguration;
+		check(FirstBuildCookRun.IsValid());
+		return FirstBuildCookRun->GetCookConfiguration();
 	}
 
 	virtual ELauncherProfileCookModes::Type GetCookMode( ) const override
 	{
-		return CookMode;
+		check(FirstBuildCookRun.IsValid());
+		return FirstBuildCookRun->GetCookMode();
 	}
 
 	virtual const FString& GetCookOptions( ) const override
 	{
-		return CookOptions;
+		check(FirstBuildCookRun.IsValid());
+		return FirstBuildCookRun->GetCookOptions();
 	}
 
 	virtual const TArray<FString>& GetCookedCultures( ) const override
 	{
-		return CookedCultures;
+		check(FirstBuildCookRun.IsValid());
+		return FirstBuildCookRun->GetCookedCultures();
 	}
 
 	virtual const bool GetSkipCookingEditorContent() const override
 	{
-		return bSkipCookingEditorContent;
+		check(FirstBuildCookRun.IsValid());
+		return FirstBuildCookRun->GetSkipCookingEditorContent();
 	}
 
 	virtual const TArray<FString>& GetCookedMaps( ) const override
 	{
-		return CookedMaps;
+		check(FirstBuildCookRun.IsValid());
+		return FirstBuildCookRun->GetCookedMaps();
 	}
 
 	virtual const TArray<FString>& GetCookedPlatforms( ) const override
 	{
-		return CookedPlatforms;
+		check(FirstBuildCookRun.IsValid());
+		return FirstBuildCookRun->GetCookedPlatforms();
 	}
 
 	virtual const ILauncherProfileLaunchRoleRef& GetDefaultLaunchRole( ) const override
 	{
-		return DefaultLaunchRole;
+		check(FirstBuildCookRun.IsValid());
+		return FirstBuildCookRun->GetDefaultLaunchRole();
 	}
 
 	virtual ILauncherDeviceGroupPtr GetDeployedDeviceGroup( ) override
 	{
-		// setting the default platform will update the device group.  always do this when getting the group because
-		// devices come in lazily through messages and can't be added properly at profile load.
-		if (DefaultDeployPlatform != NAME_None)
-		{
-			SetDefaultDeployPlatform(DefaultDeployPlatform);
-		}
-		return DeployedDeviceGroup;
+		check(FirstBuildCookRun.IsValid());
+		return FirstBuildCookRun->GetDeployedDeviceGroup();
 	}
 
 	virtual const FName GetDefaultDeployPlatform() const override
 	{
-		return DefaultDeployPlatform;
+		check(FirstBuildCookRun.IsValid());
+		return FirstBuildCookRun->GetDefaultDeployPlatform();
 	}
 
 	virtual bool IsGeneratingPatch() const override
 	{
-		return GeneratePatch;
+		check(FirstBuildCookRun.IsValid());
+		return FirstBuildCookRun->IsGeneratingPatch();
 	}
 
 	virtual bool ShouldAddPatchLevel() const override
 	{
-		return AddPatchLevel;
+		check(FirstBuildCookRun.IsValid());
+		return FirstBuildCookRun->ShouldAddPatchLevel();
 	}
 
 	virtual bool ShouldStageBaseReleasePaks() const override
@@ -512,99 +408,119 @@ public:
 
 	virtual bool IsCreatingDLC() const override
 	{
-		return CreateDLC;
+		check(FirstBuildCookRun.IsValid());
+		return FirstBuildCookRun->IsCreatingDLC();
 	}
 	virtual void SetCreateDLC(bool InBuildDLC) override
 	{
-		CreateDLC = InBuildDLC;
+		check(FirstBuildCookRun.IsValid());
+		FirstBuildCookRun->SetCreateDLC(InBuildDLC);
 	}
 
 	virtual FString GetDLCName() const override
 	{
-		return DLCName;
+		check(FirstBuildCookRun.IsValid());
+		return FirstBuildCookRun->GetDLCName();
 	}
 	virtual void SetDLCName(const FString& InDLCName) override
 	{
-		DLCName = InDLCName;
+		check(FirstBuildCookRun.IsValid());
+		FirstBuildCookRun->SetDLCName(InDLCName);
 	}
 
 	virtual bool IsDLCIncludingEngineContent() const
 	{
-		return DLCIncludeEngineContent;
+		check(FirstBuildCookRun.IsValid());
+		return FirstBuildCookRun->IsDLCIncludingEngineContent();
 	}
 	virtual void SetDLCIncludeEngineContent(bool InDLCIncludeEngineContent)
 	{
-		DLCIncludeEngineContent = InDLCIncludeEngineContent;
+		check(FirstBuildCookRun.IsValid());
+		FirstBuildCookRun->SetDLCIncludeEngineContent(InDLCIncludeEngineContent);
 	}
 
 
 
 	virtual bool IsCreatingReleaseVersion() const override
 	{
-		return CreateReleaseVersion;
+		check(FirstBuildCookRun.IsValid());
+		return FirstBuildCookRun->IsCreatingReleaseVersion();
 	}
 
 	virtual void SetCreateReleaseVersion(bool InCreateReleaseVersion) override
 	{
-		CreateReleaseVersion = InCreateReleaseVersion;
+		check(FirstBuildCookRun.IsValid());
+		FirstBuildCookRun->SetCreateReleaseVersion(InCreateReleaseVersion);
 	}
 
 	virtual FString GetCreateReleaseVersionName() const override
 	{
-		return CreateReleaseVersionName;
+		check(FirstBuildCookRun.IsValid());
+		return FirstBuildCookRun->GetCreateReleaseVersionName();
 	}
 
 	virtual void SetCreateReleaseVersionName(const FString& InCreateReleaseVersionName) override
 	{
-		CreateReleaseVersionName = InCreateReleaseVersionName;
+		check(FirstBuildCookRun.IsValid());
+		FirstBuildCookRun->SetCreateReleaseVersionName(InCreateReleaseVersionName);
 	}
 
 
 	virtual FString GetBasedOnReleaseVersionName() const override
 	{
-		return BasedOnReleaseVersionName;
+		check(FirstBuildCookRun.IsValid());
+		return FirstBuildCookRun->GetBasedOnReleaseVersionName();
 	}
 
 	virtual void SetBasedOnReleaseVersionName(const FString& InBasedOnReleaseVersionName) override
 	{
-		BasedOnReleaseVersionName = InBasedOnReleaseVersionName;
+		check(FirstBuildCookRun.IsValid());
+		FirstBuildCookRun->SetBasedOnReleaseVersionName(InBasedOnReleaseVersionName);
 	}
 
 	virtual FString GetOriginalReleaseVersionName() const override
 	{
-		return OriginalReleaseVersionName; 
+		check(FirstBuildCookRun.IsValid());
+		return FirstBuildCookRun->GetOriginalReleaseVersionName(); 
 	}
 
 	virtual void SetOriginalReleaseVersionName(const FString& InOriginalReleaseVersionName) override
 	{
-		OriginalReleaseVersionName = InOriginalReleaseVersionName;
+		check(FirstBuildCookRun.IsValid());
+		FirstBuildCookRun->SetOriginalReleaseVersionName(InOriginalReleaseVersionName);
 	}
 
 	virtual FString GetReferenceContainerGlobalFileName() const override
 	{
-		return ReferenceContainerGlobalFileName;
+		check(FirstBuildCookRun.IsValid());
+		return FirstBuildCookRun->GetReferenceContainerGlobalFileName();
 	}
 	virtual void SetReferenceContainerGlobalFileName(const FString& InReferenceContainerGlobalFileName) override
 	{
-		ReferenceContainerGlobalFileName = InReferenceContainerGlobalFileName;
+		check(FirstBuildCookRun.IsValid());
+		FirstBuildCookRun->SetReferenceContainerGlobalFileName(InReferenceContainerGlobalFileName);
 	}
 	virtual FString GetReferenceContainerCryptoKeysFileName() const override
 	{
-		return ReferenceContainerCryptoKeysFileName;
+		check(FirstBuildCookRun.IsValid());
+		return FirstBuildCookRun->GetReferenceContainerCryptoKeysFileName();
 	}
 	virtual void SetReferenceContainerCryptoKeysFileName(const FString& InReferenceContainerCryptoKeysFileName) override
 	{
-		ReferenceContainerCryptoKeysFileName = InReferenceContainerCryptoKeysFileName;
+		check(FirstBuildCookRun.IsValid());
+		FirstBuildCookRun->SetReferenceContainerCryptoKeysFileName(InReferenceContainerCryptoKeysFileName);
 	}
 
 	virtual ELauncherProfileDeploymentModes::Type GetDeploymentMode( ) const override
 	{
-		return DeploymentMode;
+		check(FirstBuildCookRun.IsValid());
+		return FirstBuildCookRun->GetDeploymentMode();
 	}
 
     virtual bool GetForceClose() const override
     {
-        return ForceClose;
+		check(FirstBuildCookRun.IsValid());
+        return FirstBuildCookRun->GetForceClose();
     }
     
 	virtual FGuid GetId( ) const override
@@ -634,36 +550,20 @@ public:
 
 	virtual ELauncherProfileLaunchModes::Type GetLaunchMode( ) const override
 	{
-		return LaunchMode;
+		check(FirstBuildCookRun.IsValid());
+		return FirstBuildCookRun->GetLaunchMode();
 	}
 
 	virtual const TArray<ILauncherProfileLaunchRolePtr>& GetLaunchRoles( ) const override
 	{
-		return LaunchRoles;
+		check(FirstBuildCookRun.IsValid());
+		return FirstBuildCookRun->GetLaunchRoles();
 	}
 
 	virtual const int32 GetLaunchRolesFor( const FString& DeviceId, TArray<ILauncherProfileLaunchRolePtr>& OutRoles ) override
 	{
-		OutRoles.Empty();
-
-		if (LaunchMode == ELauncherProfileLaunchModes::CustomRoles)
-		{
-			for (TArray<ILauncherProfileLaunchRolePtr>::TConstIterator It(LaunchRoles); It; ++It)
-			{
-				ILauncherProfileLaunchRolePtr Role = *It;
-
-				if (Role->GetAssignedDevice() == DeviceId)
-				{
-					OutRoles.Add(Role);
-				}
-			}
-		}
-		else if (LaunchMode == ELauncherProfileLaunchModes::DefaultRole)
-		{
-			OutRoles.Add(DefaultLaunchRole);
-		}
-
-		return OutRoles.Num();
+		check(FirstBuildCookRun.IsValid());
+		return FirstBuildCookRun->GetLaunchRolesFor( DeviceId, OutRoles );
 	}
 
 	virtual FString GetName( ) const override
@@ -678,22 +578,26 @@ public:
 
 	virtual ELauncherProfilePackagingModes::Type GetPackagingMode( ) const override
 	{
-		return PackagingMode;
+		check(FirstBuildCookRun.IsValid());
+		return FirstBuildCookRun->GetPackagingMode();
 	}
 
 	virtual FString GetPackageDirectory( ) const override
 	{
-		return PackageDir;
+		check(FirstBuildCookRun.IsValid());
+		return FirstBuildCookRun->GetPackageDirectory();
 	}
 
 	virtual bool IsArchiving( ) const override
 	{
-		return bArchive;
+		check(FirstBuildCookRun.IsValid());
+		return FirstBuildCookRun->IsArchiving();
 	}
 
 	virtual FString GetArchiveDirectory( ) const override
 	{
-		return ArchiveDir;
+		check(FirstBuildCookRun.IsValid());
+		return FirstBuildCookRun->GetArchiveDirectory();
 	}
 
 	virtual bool HasProjectSpecified() const override
@@ -731,148 +635,133 @@ public:
 
     virtual uint32 GetTimeout() const override
     {
-        return Timeout;
+		check(FirstBuildCookRun.IsValid());
+        return FirstBuildCookRun->GetTimeout();
     }
 
-	virtual bool HasValidationError() const override
+	virtual bool HasValidationError(ILauncherProfileUATCommandPtr UATCommand) const override
 	{
-		return ValidationErrors.Num() > 0;
+		const FValidationData* ValidationData = GetValidationData(UATCommand);
+		return ValidationData && ValidationData->Errors.Num() > 0;
 	}
     
-	virtual bool HasValidationError( ELauncherProfileValidationErrors::Type Error ) const override
+	virtual bool HasValidationError( ELauncherProfileValidationErrors::Type Error, ILauncherProfileUATCommandPtr UATCommand ) const override
 	{
-		return ValidationErrors.Contains(Error);
+		const FValidationData* ValidationData = GetValidationData(UATCommand);
+		return ValidationData && ValidationData->Errors.Contains(Error);
 	}
 
-	virtual void AddCustomError( const FString& UniqueId, const FText& Text ) override
+	virtual void AddCustomError( const FString& UniqueId, const FText& Text, ILauncherProfileUATCommandPtr UATCommand ) override
 	{
-		CustomErrors.Add(UniqueId, Text);
+		FValidationData* ValidationData = GetValidationData(UATCommand);
+		if (ValidationData)
+		{
+			ValidationData->CustomErrors.Add(UniqueId, Text);
+		}
 	}
 
-	virtual bool HasCustomError( const FString& UniqueId ) const override
+	virtual bool HasCustomError( const FString& UniqueId, ILauncherProfileUATCommandPtr UATCommand ) const override
 	{
-		return CustomErrors.Contains(UniqueId);
+		const FValidationData* ValidationData = GetValidationData(UATCommand);
+		return ValidationData && ValidationData->CustomErrors.Contains(UniqueId);
 	}
 
-	virtual FText GetCustomErrorText( const FString& UniqueId ) const override
+	virtual FText GetCustomErrorText( const FString& UniqueId, ILauncherProfileUATCommandPtr UATCommand ) const override
 	{
-		const FText* TextPtr = CustomErrors.Find(UniqueId);
+		const FValidationData* ValidationData = GetValidationData(UATCommand);
+		const FText* TextPtr = ValidationData ? ValidationData->CustomErrors.Find(UniqueId) : nullptr;
 		return TextPtr ? (*TextPtr) : FText::GetEmpty();
 	}
 
-	virtual TArray<FString> GetAllCustomErrors() const override
+	virtual TArray<FString> GetAllCustomErrors(ILauncherProfileUATCommandPtr UATCommand) const override
 	{
 		TArray<FString> Keys;
-		CustomErrors.GetKeys(Keys);
+
+		const FValidationData* ValidationData = GetValidationData(UATCommand);
+		if (ValidationData)
+		{
+			ValidationData->CustomErrors.GetKeys(Keys);
+		}
+
 		return MoveTemp(Keys);
 	}
 
-	virtual void AddCustomWarning( const FString& UniqueId, const FText& Text ) override
+	virtual void AddCustomWarning( const FString& UniqueId, const FText& Text, ILauncherProfileUATCommandPtr UATCommand ) override
 	{
-		CustomWarnings.Add(UniqueId, Text);
+		FValidationData* ValidationData = GetValidationData(UATCommand);
+		if (ValidationData)
+		{
+			ValidationData->CustomWarnings.Add(UniqueId, Text);
+		}
 	}
 
-	virtual bool HasCustomWarning( const FString& UniqueId ) const override
+	virtual bool HasCustomWarning( const FString& UniqueId, ILauncherProfileUATCommandPtr UATCommand ) const override
 	{
-		return CustomWarnings.Contains(UniqueId);
+		const FValidationData* ValidationData = GetValidationData(UATCommand);
+		return ValidationData && ValidationData->CustomWarnings.Contains(UniqueId);
 	}
 
-	virtual FText GetCustomWarningText( const FString& UniqueId ) const override
+	virtual FText GetCustomWarningText( const FString& UniqueId, ILauncherProfileUATCommandPtr UATCommand ) const override
 	{
-		const FText* TextPtr = CustomWarnings.Find(UniqueId);
+		const FValidationData* ValidationData = GetValidationData(UATCommand);
+		const FText* TextPtr = ValidationData ? ValidationData->CustomWarnings.Find(UniqueId) : nullptr;
 		return TextPtr ? (*TextPtr) : FText::GetEmpty();
 	}
 
-	virtual TArray<FString> GetAllCustomWarnings() const override
+	virtual TArray<FString> GetAllCustomWarnings(ILauncherProfileUATCommandPtr UATCommand) const override
 	{
 		TArray<FString> Keys;
-		CustomWarnings.GetKeys(Keys);
+		const FValidationData* ValidationData = GetValidationData(UATCommand);
+		if (ValidationData)
+		{
+			ValidationData->CustomWarnings.GetKeys(Keys);
+		}
 		return MoveTemp(Keys);
 	}
 
-	virtual void RefreshCustomWarningsAndErrors() override
+	virtual void RefreshCustomWarningsAndErrors(ILauncherProfileUATCommandPtr UATCommand) override
 	{
-		CustomWarnings.Reset();
-		CustomErrors.Reset();
-		CustomValidationDelegate.Broadcast();
+		FValidationData* ValidationData = GetValidationData(UATCommand);
+		if (ValidationData)
+		{
+			ValidationData->CustomWarnings.Reset();
+			ValidationData->CustomErrors.Reset();
+		}
+
+		if (UATCommand == nullptr)
+		{
+			CustomValidationDelegate.Broadcast();
+		}
+		else
+		{
+			CustomUATCommandValidationDelegate.Broadcast(UATCommand.ToSharedRef());
+		}
 	}
 
 
 
-	virtual FString GetInvalidPlatform() const override
+	virtual FString GetInvalidPlatform(ILauncherProfileUATCommandPtr UATCommand) const override
 	{
-		return InvalidPlatform;
+		const FValidationData* ValidationData = GetValidationData(UATCommand);
+		return ValidationData ? ValidationData->InvalidPlatform : FString();
 	}
 
 	virtual ELauncherProfileBuildModes::Type GetBuildMode() const override
 	{
-		return BuildMode;
+		check(FirstBuildCookRun.IsValid());
+		return FirstBuildCookRun->GetBuildMode();
 	}
 
 	virtual bool ShouldBuild() override
 	{
-		bool bBuild = true;
-		if (BuildMode == ELauncherProfileBuildModes::DoNotBuild)
-		{
-			bBuild = false;
-		}
-		else if (BuildMode == ELauncherProfileBuildModes::Auto)
-		{
-			if (FApp::GetEngineIsPromotedBuild())
-			{
-				bBuild = false;
-
-				TArray<FString> TargetPlatformNames = FindPlatforms();
-				for (const FString& TargetPlatformName : TargetPlatformNames)
-				{
-					// Get the target we're building for
-					const ITargetPlatform* TargetPlatform = GetTargetPlatformManager()->FindTargetPlatform(TargetPlatformName);
-					const PlatformInfo::FTargetPlatformInfo& PlatformInfo = TargetPlatform->GetTargetPlatformInfo();
-
-					// Figure out which target we're building
-					FString ReceiptDir;
-					FString TargetName;
-					if (TryGetDefaultTargetName(FPaths::GetProjectFilePath(), PlatformInfo.PlatformType, TargetName))
-					{
-						ReceiptDir = FPaths::GetPath(FPaths::GetProjectFilePath());
-					}
-					else if (TryGetDefaultTargetName(FString(), PlatformInfo.PlatformType, TargetName))
-					{
-						FText Reason;
-						if (TargetPlatform->RequiresTempTarget(false, BuildConfiguration, false, Reason))
-						{
-							UE_LOG(LogLauncherProfile, Log, TEXT("Project requires temp target (%s)"), *Reason.ToString());
-							ReceiptDir = FPaths::GetPath(FPaths::GetProjectFilePath());
-						}
-						else
-						{
-							UE_LOG(LogLauncherProfile, Log, TEXT("Project does not require temp target"));
-							ReceiptDir = FPaths::EngineDir();
-						}
-					}
-					else
-					{
-						UE_LOG(LogLauncherProfile, Log, TEXT("Unable to find any targets for platform %s - forcing build"), *TargetPlatformName);
-						bBuild = true;
-						break;
-					}
-
-					// Check if the existing target is valid
-					FString BuildPlatform = PlatformInfo.DataDrivenPlatformInfo->UBTPlatformString;
-					if (!HasPromotedTarget(*ReceiptDir, *TargetName, *BuildPlatform, BuildConfiguration, nullptr))
-					{
-						bBuild = true;
-						break;
-					}
-				}
-			}
-		}
-		return bBuild;
+		check(FirstBuildCookRun.IsValid());
+		return FirstBuildCookRun->ShouldBuild();
 	}
 
 	virtual bool IsFastIterate() const override
 	{
-		return bFastIterate;
+		check(FirstBuildCookRun.IsValid());
+		return FirstBuildCookRun->IsFastIterate();
 	}
 	
 	virtual bool IsBuildingUAT() const override
@@ -882,139 +771,157 @@ public:
 
 	virtual FString GetAdditionalCommandLineParameters() const override
 	{
-		return AdditionalCommandLineParameters;
+		check(FirstBuildCookRun.IsValid());
+		return FirstBuildCookRun->GetAdditionalCommandLineParameters();
 	}
 
 	virtual FString GetAdditionalTargetCommandLineParameters( EBuildTargetType BuildTargetType ) const override
 	{
-		check( BuildTargetType != EBuildTargetType::Unknown);
-
-		const FString* CommandLinePtr = AdditionalTargetCommandLineParameters.Find(BuildTargetType);
-		return CommandLinePtr ? *CommandLinePtr : FString();
+		check(FirstBuildCookRun.IsValid());
+		return FirstBuildCookRun->GetAdditionalTargetCommandLineParameters(BuildTargetType);
 	}
 
 	virtual bool IsCookingIncrementally( ) const override
 	{
-		if ( CookMode != ELauncherProfileCookModes::DoNotCook )
-		{
-			return CookIncremental;
-		}
-		return false;
+		check(FirstBuildCookRun.IsValid());
+		return FirstBuildCookRun->IsCookingIncrementally();
 	}
 
 	virtual ELauncherProfileIncrementalCookMode::Type GetIncrementalCookMode() const override
 	{
 		// fixme: worth reading LegacyIterative to honor this value?!
-		return CookIncremental ? ELauncherProfileIncrementalCookMode::ModifiedAndDependencies : ELauncherProfileIncrementalCookMode::None;
+		return IsCookingIncrementally() ? ELauncherProfileIncrementalCookMode::ModifiedAndDependencies : ELauncherProfileIncrementalCookMode::None;
 	}
 
 
 	virtual bool IsIterateSharedCookedBuild() const override
 	{
-		if ( CookMode != ELauncherProfileCookModes::DoNotCook )
-		{
-			return IterateSharedCookedBuild;
-		}
-		return false;
+		check(FirstBuildCookRun.IsValid());
+		return FirstBuildCookRun->IsIterateSharedCookedBuild();
 	}
 
 	virtual bool IsCompressed() const override
 	{
-		return Compressed;
+		check(FirstBuildCookRun.IsValid());
+		return FirstBuildCookRun->IsCompressed();
 	}
 
 	virtual bool IsEncryptingIniFiles() const override
 	{
-		return EncryptIniFiles;
+		check(FirstBuildCookRun.IsValid());
+		return FirstBuildCookRun->IsEncryptingIniFiles();
 	}
 
 	virtual bool IsForDistribution() const override
 	{
-		return ForDistribution;
+		check(FirstBuildCookRun.IsValid());
+		return FirstBuildCookRun->IsForDistribution();
 	}
 
 	virtual bool IsCookingUnversioned( ) const override
 	{
-		return CookUnversioned;
+		check(FirstBuildCookRun.IsValid());
+		return FirstBuildCookRun->IsCookingUnversioned();
 	}
 
 	virtual bool IsDeployablePlatform( const FString& PlatformName ) override
 	{
-		if (CookMode == ELauncherProfileCookModes::ByTheBook || CookMode == ELauncherProfileCookModes::ByTheBookInEditor)
-		{
-			return CookedPlatforms.Contains(PlatformName);
-		}
-
-		return true;
+		check(FirstBuildCookRun.IsValid());
+		return FirstBuildCookRun->IsDeployablePlatform(PlatformName);
 	}
 
 	virtual bool IsDeployingIncrementally( ) const override
 	{
-		return DeployIncremental;
+		check(FirstBuildCookRun.IsValid());
+		return FirstBuildCookRun->IsDeployingIncrementally();
 	}
 
 	virtual bool IsFileServerHidden( ) const override
 	{
-		return HideFileServerWindow;
+		check(FirstBuildCookRun.IsValid());
+		return FirstBuildCookRun->IsFileServerHidden();
 	}
 
 	virtual bool IsFileServerStreaming( ) const override
 	{
-		return DeployStreamingServer;
+		check(FirstBuildCookRun.IsValid());
+		return FirstBuildCookRun->IsFileServerStreaming();
 	}
 
 	virtual bool IsPackingWithUnrealPak( ) const  override
 	{
-		return DeployWithUnrealPak;
+		check(FirstBuildCookRun.IsValid());
+		return FirstBuildCookRun->IsPackingWithUnrealPak();
 	}
 
 	virtual bool IsIncludingPrerequisites() const override
 	{
-		return IncludePrerequisites;
+		check(FirstBuildCookRun.IsValid());
+		return FirstBuildCookRun->IsIncludingPrerequisites();
 	}
 
 	virtual bool IsGeneratingChunks() const override
 	{
-		return bGenerateChunks;
+		check(FirstBuildCookRun.IsValid());
+		return FirstBuildCookRun->IsGeneratingChunks();
 	}
 
 	virtual bool IsGenerateHttpChunkData() const override
 	{
-		return bGenerateHttpChunkData;
+		check(FirstBuildCookRun.IsValid());
+		return FirstBuildCookRun->IsGenerateHttpChunkData();
 	}
 
 	virtual FString GetHttpChunkDataDirectory() const override
 	{
-		return HttpChunkDataDirectory;
+		check(FirstBuildCookRun.IsValid());
+		return FirstBuildCookRun->GetHttpChunkDataDirectory();
 	}
 
 	virtual FString GetHttpChunkDataReleaseName() const override
 	{
-		return HttpChunkDataReleaseName;
+		check(FirstBuildCookRun.IsValid());
+		return FirstBuildCookRun->GetHttpChunkDataReleaseName();
 	}
 
 	virtual bool IsValidForLaunch( ) override
 	{
-		return (ValidationErrors.Num() == 0 && CustomErrors.Num() == 0);
+		if (ProfileValidation.Errors.Num() != 0 || ProfileValidation.CustomErrors.Num() != 0)
+		{
+			return false;
+		}
+
+		for (const TPair<ILauncherProfileUATCommandPtr,FValidationData>& Pair : UATCommandValidation)
+		{
+			if (Pair.Value.Errors.Num() != 0 || Pair.Value.CustomErrors.Num() != 0)
+			{
+				return false;
+			}
+		}
+
+		return true;
 	}
 
 	virtual void RemoveCookedCulture( const FString& CultureName ) override
 	{
-		CookedCultures.Remove(CultureName);
+		check(FirstBuildCookRun.IsValid());
+		FirstBuildCookRun->RemoveCookedCulture(CultureName);
 
 		Validate();
 	}
 
 	virtual void RemoveCookedMap( const FString& MapName ) override
 	{
-		CookedMaps.Remove(MapName);
+		check(FirstBuildCookRun.IsValid());
+		FirstBuildCookRun->RemoveCookedMap(MapName);
 
 		Validate();
 	}
 
 	virtual void RemoveCookedPlatform( const FString& PlatformName ) override
 	{
-		CookedPlatforms.Remove(PlatformName);
+		check(FirstBuildCookRun.IsValid());
+		FirstBuildCookRun->RemoveCookedPlatform(PlatformName);
 
 		RefreshValidBuildTargets();
 		Validate();
@@ -1022,467 +929,29 @@ public:
 
 	virtual void RemoveLaunchRole( const ILauncherProfileLaunchRoleRef& Role ) override
 	{
-		LaunchRoles.Remove(Role);
-
-		Validate();
-	}
-
-	UE_DEPRECATED(5.8, "please use the Json Save/Load API instead")
-	virtual bool Serialize( FArchive& Archive ) override
-	{
-		int32 Version = LAUNCHERSERVICES_FINAL;
-
-		Archive	<< Version;
-
-		if (Version < LAUNCHERSERVICES_MINPROFILEVERSION)
-		{
-			return false;
-		}
-
-		if (Version > LAUNCHERSERVICES_FINAL)
-		{
-			return false;
-		}
-
-		if (Archive.IsSaving())
-		{
-			if (DeployedDeviceGroup.IsValid())
-			{
-				DeployedDeviceGroupId = DeployedDeviceGroup->GetId();
-			}
-			else
-			{
-				DeployedDeviceGroupId = FGuid();
-			}
-		}
-
-		// IMPORTANT: make sure to bump LAUNCHERSERVICES_PROFILEVERSION when modifying this!
-		bool BuildGame = false;
-		Archive << Id
-				<< Name
-				<< Description
-				<< BuildConfiguration
-				<< ProjectSpecified
-				<< ShareableProjectPath
-				<< CookConfiguration
-				<< CookIncremental
-				<< CookOptions
-				<< CookMode
-				<< CookUnversioned
-				<< CookedCultures
-				<< CookedMaps
-				<< CookedPlatforms
-				<< DeployStreamingServer
-				<< DeployWithUnrealPak
-				<< DeployedDeviceGroupId
-				<< DeploymentMode
-				<< HideFileServerWindow
-				<< LaunchMode
-				<< PackagingMode
-				<< PackageDir
-				<< BuildGame
-                << ForceClose
-                << Timeout;
-
-		if (Version >= LAUNCHERSERVICES_SHAREABLEPROJECTPATHS)
-		{
-			FullProjectPath = FPaths::ConvertRelativePathToFull(FPaths::RootDir(), ShareableProjectPath);
-		}
-
-		FString DeployPlatformString = DefaultDeployPlatform.ToString();
-		if (Version >= LAUNCHERSERVICES_FIXCOMPRESSIONSERIALIZE)
-		{
-			Archive << Compressed;
-		}
-		if ( Version>= LAUNCHERSERVICES_ADDEDENCRYPTINIFILES)
-		{
-			Archive << EncryptIniFiles;
-			Archive << ForDistribution;
-		}
-		if (Version >= LAUNCHERSERVICES_ADDEDDEFAULTDEPLOYPLATFORM)
-		{
-			Archive << DeployPlatformString;
-		}		
-		if (Version >= LAUNCHERSERVICES_ADDEDNUMCOOKERSTOSPAWN && Version < LAUNCHERSERVICES_REMOVEDNUMCOOKERSTOSPAWN)
-		{
-			int32 OldNumCookersToSpawn;
-			Archive << OldNumCookersToSpawn;
-		}
-		if (Version >= LAUNCHERSERVICES_ADDEDSKIPCOOKINGEDITORCONTENT)
-		{
-			Archive << bSkipCookingEditorContent;
-		}
-		if (Version >= LAUNCHERSERVICES_ADDEDINCREMENTALDEPLOYVERSION)
-		{
-			Archive << DeployIncremental;
-		}
-		if ( Version >= LAUNCHERSERVICES_REMOVEDPATCHSOURCECONTENTPATH )
-		{
-			Archive << GeneratePatch;
-		}
-		if ( Version >= LAUNCHERSERVICES_ADDEDMULTILEVELPATCHING )
-		{
-			Archive << AddPatchLevel;
-			Archive << StageBaseReleasePaks;
-		}
-		else if ( Version >= LAUNCHERSERVICES_ADDEDPATCHSOURCECONTENTPATH)
-		{
-			FString Temp;
-			Archive << Temp;
-			Archive << GeneratePatch;
-		}
-		if (Version >= LAUNCHERSERVICES_ADDEDDLCINCLUDEENGINECONTENT)
-		{
-			Archive << DLCIncludeEngineContent;
-		}
-		
-		if ( Version >= LAUNCHERSERVICES_ADDEDRELEASEVERSION )
-		{
-			Archive << CreateReleaseVersion;
-			Archive << CreateReleaseVersionName;
-			Archive << BasedOnReleaseVersionName;
-			
-			Archive << CreateDLC;
-			Archive << DLCName;
-		}
-		if (Version >= LAUNCHERSERVICES_ADDEDGENERATECHUNKS)
-		{
-			Archive << bGenerateChunks;
-			Archive << bGenerateHttpChunkData;
-			Archive << HttpChunkDataDirectory;
-			Archive << HttpChunkDataReleaseName;
-		}
-		if (Version >= LAUNCHERSERVICES_ADDARCHIVE)
-		{
-			Archive << bArchive;
-			Archive << ArchiveDir;
-		}
-		if (Version >= LAUNCHERSERVICES_ADDEDADDITIONALCOMMANDLINE)
-		{
-			Archive << AdditionalCommandLineParameters;
-		}
-		if (Version >= LAUNCHERSERVICES_ADDEDINCLUDEPREREQUISITES)
-		{
-			Archive << IncludePrerequisites;
-		}
-
-		if (Version >= LAUNCHERSERVICES_ADDEDBUILDMODE)
-		{
-			Archive << BuildMode;
-		}
-		else if(Archive.IsLoading())
-		{
-			BuildMode = BuildGame ? ELauncherProfileBuildModes::Build : ELauncherProfileBuildModes::DoNotBuild;
-		}
-
-		if (Version >= LAUNCHERSERVICES_ADDEDUSEIOSTORE)
-		{
-			Archive << bUseIoStore;
-		}
-
-		if (Version >= LAUNCHERSERVICES_ADDEDMAKEBINARYCONFIG)
-		{
-			Archive << bMakeBinaryConfig;
-		}
-
-		if (Version >= LAUNCHERSERVICES_ADDEDREFERENCECONTAINERS)
-		{
-			Archive << ReferenceContainerGlobalFileName;
-			Archive << ReferenceContainerCryptoKeysFileName;
-		}
-		if (Version >= LAUNCHERSERVICES_ADDEDORIGINALRELEASEVERSION)
-		{
-			Archive << OriginalReleaseVersionName;
-		}
-
-		FString BuildTargetName;
-		if (Version >= LAUNCHERSERVICES_ADDBUILDTARGETNAME)
-		{
-			Archive << BuildTargetSpecified;
-			Archive << BuildTargetName;
-		}
-
-		if (Version >= LAUNCHERSERVICES_ADDEDRETAINSTAGEDDIRECTORY && Version < LAUNCHERSERVICES_REMOVEDRETAINSTAGEDDIRECTORY)
-		{
-			bool bRetainStagedDirectory = false;
-			Archive << bRetainStagedDirectory;
-		}
-
-		if (Version >= LAUNCHERSERVICES_ADDZENOPTIONS)
-		{
-			Archive << bUseZenStore;
-			Archive << bImportingZenSnapshot;
-			Archive << bUseZenPakStreaming;
-			Archive << ZenPakStreamingPath;
-		}
-
-		if (Version >= LAUNCHERSERVICES_ADDBUILDARCHITECTURE)
-		{
-			Archive << ClientArchitectures;
-			Archive << ServerArchitectures;
-			Archive << EditorArchitectures;
-		}
-
-		if (Version >= LAUNCHERSERVICES_ADDCUSTOMPROPERTIES)
-		{
-			Archive << CustomStringProperties;
-			Archive << CustomBoolProperties;
-		}
-
-		if (Version >= LAUNCHERSERVICES_INCREMENTALCOOKMODE && Version < LAUNCHERSERVICES_REMOVEDINCREMENTALCOOKMODE)
-		{
-			TEnumAsByte<ELauncherProfileIncrementalCookMode::Type> IncrementalCookMode = (CookIncremental ? ELauncherProfileIncrementalCookMode::ModifiedOnly : ELauncherProfileIncrementalCookMode::None);
-			Archive << IncrementalCookMode;
-		}
-
-		if (Version >= LAUNCHERSERVICES_AUTOMATEDTESTS)
-		{
-			TArray<TSharedPtr<ILauncherProfileAutomatedTest>> AutomatedTests;
-			if (Archive.IsSaving())
-			{
-				for (const ILauncherProfileUATCommandRef& UATCommand : UATCommands)
-				{
-					ILauncherProfileAutomatedTestPtr AutomatedTest = UATCommand->AsAutomatedTest();
-					if (AutomatedTest)
-					{
-						AutomatedTests.Add(AutomatedTest);
-					}
-				}
-			}
-
-			// serialize automated tests
-			int NumAutomatedTests = AutomatedTests.Num();
-			Archive << NumAutomatedTests;
-			for( int32 TestIndex = 0; TestIndex < NumAutomatedTests; TestIndex++)
-			{
-PRAGMA_DISABLE_DEPRECATION_WARNINGS
-				if (Archive.IsLoading())
-				{
-					AutomatedTests.Add(MakeShared<FLauncherProfileAutomatedTest>(Archive));
-				}
-				else
-				{
-					AutomatedTests[TestIndex]->Serialize(Archive);
-				}
-			}
-PRAGMA_ENABLE_DEPRECATION_WARNINGS
-
-			if (Archive.IsLoading())
-			{
-				for (TSharedPtr<ILauncherProfileAutomatedTest> AutomatedTest : AutomatedTests)
-				{
-					UATCommands.Add(AutomatedTest.ToSharedRef());
-				}
-			}
-
-			Archive << AutomatedTestBuildPath;
-			Archive << bWantAutomatedTestBuild;
-		}
-
-		if (Version >= LAUNCHERSERVICES_ADDMULTIPLEBUILDTARGETS)
-		{
-			Archive << BuildTargets;
-			if (Archive.IsLoading() && BuildTargets.Num() == 0 && !BuildTargetName.IsEmpty())
-			{
-				BuildTargets.Add(BuildTargetName);
-			}
-
-			Archive << AdditionalTargetCommandLineParameters;
-		}
-
-PRAGMA_DISABLE_DEPRECATION_WARNINGS
-		DefaultLaunchRole->Serialize(Archive);
-PRAGMA_ENABLE_DEPRECATION_WARNINGS
-
-		// serialize launch roles
-		if (Archive.IsLoading())
-		{
-			DeployedDeviceGroup.Reset();
-			LaunchRoles.Reset();
-		}
-
-		int32 NumLaunchRoles = LaunchRoles.Num();
-
-		Archive << NumLaunchRoles;
-
-		for (int32 RoleIndex = 0; RoleIndex < NumLaunchRoles; ++RoleIndex)
-		{
-PRAGMA_DISABLE_DEPRECATION_WARNINGS
-			if (Archive.IsLoading())
-			{
-				LaunchRoles.Add(MakeShareable(new FLauncherProfileLaunchRole(Archive)));				
-			}
-			else
-			{
-				LaunchRoles[RoleIndex]->Serialize(Archive);
-			}
-PRAGMA_ENABLE_DEPRECATION_WARNINGS
-		}
-
-		if (Archive.IsLoading())
-		{
-			DefaultDeployPlatform = FName(*DeployPlatformString);
-		}
-
-		if (DefaultDeployPlatform != NAME_None)
-		{
-			SetDefaultDeployPlatform(DefaultDeployPlatform);
-		}
-
-		if (Archive.IsLoading())
-		{
-			RefreshValidBuildTargets();
-		}
-		Validate();
-
-		return true;
-	}
-
-	virtual void Save(TJsonWriter<>& Writer) override
-	{
-		int32 Version = LAUNCHERSERVICES_FINAL;
-
-		if (DeployedDeviceGroup.IsValid())
-		{
-			DeployedDeviceGroupId = DeployedDeviceGroup->GetId();
-		}
-		else
-		{
-			DeployedDeviceGroupId = FGuid();
-		}
-
-		Writer.WriteObjectStart();
-		Writer.WriteValue("Version", Version);
-		Writer.WriteValue("Id", Id.ToString());
-		Writer.WriteValue("Name", Name);
-		Writer.WriteValue("Description", Description);
-		Writer.WriteValue("BuildConfiguration", (int32)BuildConfiguration);
-		Writer.WriteValue("ProjectSpecified", ProjectSpecified);
-		Writer.WriteValue("ShareableProjectPath", ShareableProjectPath);
-		Writer.WriteValue("CookConfiguration", (int32)CookConfiguration);
-		Writer.WriteValue("CookIncremental", CookIncremental);
-		Writer.WriteValue("CookOptions", CookOptions);
-		Writer.WriteValue("CookMode", CookMode);
-		Writer.WriteValue("CookUnversioned", CookUnversioned);
-
-		// write the cultures
-		if (CookedCultures.Num() > 0)
-		{
-			Writer.WriteArrayStart("CookedCultures");
-			for (auto Value : CookedCultures)
-			{
-				Writer.WriteValue(Value);
-			}
-			Writer.WriteArrayEnd();
-		}
-
-		// write the maps
-		if (CookedMaps.Num() > 0)
-		{
-			Writer.WriteArrayStart("CookedMaps");
-			for (auto Value : CookedMaps)
-			{
-				Writer.WriteValue(Value);
-			}
-			Writer.WriteArrayEnd();
-		}
-
-		// write the platforms
-		if (CookedPlatforms.Num() > 0)
-		{
-			Writer.WriteArrayStart("CookedPlatforms");
-			for (auto Value : CookedPlatforms)
-			{
-				Writer.WriteValue(Value);
-			}
-			Writer.WriteArrayEnd();
-		}
-
-		// write the build targets
-		if (BuildTargets.Num() > 0)
-		{
-			Writer.WriteArrayStart("BuildTargets");
-			for (auto Value : BuildTargets)
-			{
-				Writer.WriteValue(Value);
-			}
-			Writer.WriteArrayEnd();
-		}
+		check(FirstBuildCookRun.IsValid());
+		FirstBuildCookRun->RemoveLaunchRole(Role);
 
-		Writer.WriteValue("DeployStreamingServer", DeployStreamingServer);
-		Writer.WriteValue("DeployWithUnrealPak", DeployWithUnrealPak);
-		Writer.WriteValue("DeployedDeviceGroupId", DeployedDeviceGroupId.ToString());
-		Writer.WriteValue("DeploymentMode", DeploymentMode);
-		Writer.WriteValue("HideFileServerWindow", HideFileServerWindow);
-		Writer.WriteValue("LaunchMode", LaunchMode);
-		Writer.WriteValue("PackagingMode", PackagingMode);
-		Writer.WriteValue("PackageDir", PackageDir);
-		Writer.WriteValue("BuildMode", BuildMode);
-		Writer.WriteValue("ForceClose", ForceClose);
-		Writer.WriteValue("Timeout", (int32)Timeout);
-		Writer.WriteValue("Compressed", Compressed);
-		Writer.WriteValue("EncryptIniFiles", EncryptIniFiles);
-		Writer.WriteValue("ForDistribution", ForDistribution);
-		Writer.WriteValue("DeployPlatform", DefaultDeployPlatform.ToString());
-		Writer.WriteValue("SkipCookingEditorContent", bSkipCookingEditorContent);
-		Writer.WriteValue("DeployIncremental", DeployIncremental);
-		Writer.WriteValue("GeneratePatch", GeneratePatch);
-		Writer.WriteValue("AddPatchLevel", AddPatchLevel);
-		Writer.WriteValue("StageBaseReleasePaks", StageBaseReleasePaks);
-		Writer.WriteValue("DLCIncludeEngineContent", DLCIncludeEngineContent);
-		Writer.WriteValue("CreateReleaseVersion", CreateReleaseVersion);
-		Writer.WriteValue("CreateReleaseVersionName", CreateReleaseVersionName);
-		Writer.WriteValue("BasedOnReleaseVersionName", BasedOnReleaseVersionName);
-		Writer.WriteValue("ReferenceContainerGlobalFileName", ReferenceContainerGlobalFileName);
-		Writer.WriteValue("ReferenceContainerCryptoKeysFileName", ReferenceContainerCryptoKeysFileName);
-		Writer.WriteValue("OriginalReleaseVersionName", OriginalReleaseVersionName);
-		Writer.WriteValue("CreateDLC", CreateDLC);
-		Writer.WriteValue("DLCName", DLCName);
-		Writer.WriteValue("GenerateChunks", bGenerateChunks);
-		Writer.WriteValue("GenerateHttpChunkData", bGenerateHttpChunkData);
-		Writer.WriteValue("HttpChunkDataDirectory", HttpChunkDataDirectory);
-		Writer.WriteValue("HttpChunkDataReleaseName", HttpChunkDataReleaseName);
-		Writer.WriteValue("Archive", bArchive);
-		Writer.WriteValue("ArchiveDirectory", ArchiveDir);
-		Writer.WriteValue("AdditionalCommandLineParameters", AdditionalCommandLineParameters);
-		Writer.WriteValue("IncludePrerequisites", IncludePrerequisites);
-		Writer.WriteValue("UseIoStore", bUseIoStore);
-		Writer.WriteValue("MakeBinaryConfig", bMakeBinaryConfig);
-		Writer.WriteValue("BuildTargetSpecified", BuildTargetSpecified);
-		Writer.WriteValue("UseZenStore", bUseZenStore);
-		Writer.WriteValue("ImportingZenSnapshot", bImportingZenSnapshot);
-		Writer.WriteValue("UseZenPakStreaming", bUseZenPakStreaming);
-		Writer.WriteValue("ZenPakStreamingPath", ZenPakStreamingPath);
-		Writer.WriteValue("ClientArchitectures", ClientArchitectures);
-		Writer.WriteValue("ServerArchitectures", ServerArchitectures);
-		Writer.WriteValue("EditorArchitectures", EditorArchitectures);
-		Writer.WriteValue("CustomStringProperties", CustomStringProperties);
-		Writer.WriteValue("CustomBoolProperties", CustomBoolProperties);
-		Writer.WriteValue("FastIterate", bFastIterate);
+		Validate();
+	}
 
-		if (AdditionalTargetCommandLineParameters.Num() > 0)
-		{
-			Writer.WriteObjectStart("AdditionalTargetCommandLineParameters");
-			for (const TTuple<EBuildTargetType,FString>& Itr : AdditionalTargetCommandLineParameters)
-			{
-				Writer.WriteValue(LexToString(Itr.Key), Itr.Value);
-			}
-			Writer.WriteObjectEnd();
-		}
+	UE_DEPRECATED(5.8, "please use the Json Save/Load API instead")
+	virtual bool Serialize( FArchive& Archive ) override;
 
-		// serialize the default launch role
-		DefaultLaunchRole->Save(Writer, TEXT("DefaultRole"));
+	virtual void Save(TJsonWriter<>& Writer) override
+	{
+		int32 Version = LAUNCHERSERVICES_FINAL;
 
-		// serialize the launch roles
-		if (LaunchRoles.Num())
-		{
-			Writer.WriteArrayStart("LaunchRoles");
-			for (auto Value : LaunchRoles)
-			{
-				Value->Save(Writer);
-			}
-			Writer.WriteArrayEnd();
-		}
+		Writer.WriteObjectStart();
+		Writer.WriteValue("Version", Version);
+		Writer.WriteValue("Id", Id.ToString());
+		Writer.WriteValue("Name", Name);
+		Writer.WriteValue("Description", Description);
+		Writer.WriteValue("ProjectSpecified", ProjectSpecified);
+		Writer.WriteValue("ShareableProjectPath", ShareableProjectPath);
+		Writer.WriteValue("StageBaseReleasePaks", StageBaseReleasePaks); // deprecated in 5.6
+		Writer.WriteValue("CustomStringProperties", CustomStringProperties);
+		Writer.WriteValue("CustomBoolProperties", CustomBoolProperties);
 
 		// serialize the automated test info
 		Writer.WriteValue("AutomatedTestBuildPath", AutomatedTestBuildPath);
@@ -1497,7 +966,10 @@ PRAGMA_ENABLE_DEPRECATION_WARNINGS
 		Writer.WriteArrayEnd();
 
 		// write out the UAT project params
-		SaveUATParams(Writer);
+		if (FirstBuildCookRun.IsValid())
+		{
+			SaveUATParams(Writer);
+		}
 		Writer.WriteObjectEnd();
 	}
 
@@ -1523,12 +995,12 @@ PRAGMA_ENABLE_DEPRECATION_WARNINGS
 
 		// client configurations
 		Writer.WriteArrayStart("clientconfig");
-		Writer.WriteValue(LexToString(BuildConfiguration));
+		Writer.WriteValue(LexToString(GetBuildConfiguration()));
 		Writer.WriteArrayEnd();
 
 		// server configurations
 		Writer.WriteArrayStart("serverconfig");
-		Writer.WriteValue(LexToString(BuildConfiguration));
+		Writer.WriteValue(LexToString(GetBuildConfiguration()));
 		Writer.WriteArrayEnd();
 
 		// platforms
@@ -1806,17 +1278,17 @@ PRAGMA_ENABLE_DEPRECATION_WARNINGS
 			}
 		}
 
-		if (ClientArchitectures.Num() > 0)
+		if (GetClientArchitectures().Num() > 0)
 		{
-			Writer.WriteValue("clientarchitecture", FString::Join(ClientArchitectures, TEXT("+")));
+			Writer.WriteValue("clientarchitecture", FString::Join(GetClientArchitectures(), TEXT("+")));
 		}
-		if (ServerArchitectures.Num() > 0)
+		if (GetServerArchitectures().Num() > 0)
 		{
-			Writer.WriteValue("serverarchitecture", FString::Join(ServerArchitectures, TEXT("+")));
+			Writer.WriteValue("serverarchitecture", FString::Join(GetServerArchitectures(), TEXT("+")));
 		}
-		if (EditorArchitectures.Num() > 0)
+		if (GetEditorArchitectures().Num() > 0)
 		{
-			Writer.WriteValue("editorarchitecture", FString::Join(EditorArchitectures, TEXT("+")));
+			Writer.WriteValue("editorarchitecture", FString::Join(GetEditorArchitectures(), TEXT("+")));
 		}
 
 		/*
@@ -2032,196 +1504,26 @@ PRAGMA_ENABLE_DEPRECATION_WARNINGS
 			return false;
 		}
 		UATCommands.Reset();
+		UATCommandValidation.Reset();
+		FirstBuildCookRun.Reset();
+		
+		if (Version < LAUNCHERSERVICES_INNERBUILDCOOKRUN)
+		{
+			ILauncherProfileBuildCookRunPtr NewInner = LoadBuildCookRunInternal(Object, Version);
+			if (!NewInner.IsValid())
+			{
+				return false;
+			}
+
+			FirstBuildCookRun = NewInner;
+			AddUATCommandInternal(NewInner.ToSharedRef());
+		}		
 
 		FGuid::Parse(Object.GetStringField(TEXT("Id")), Id);
 		Name = Object.GetStringField(TEXT("Name"));
 		Description = Object.GetStringField(TEXT("Description"));
-		BuildConfiguration = (EBuildConfiguration)((int32)Object.GetNumberField(TEXT("BuildConfiguration")));
 		ProjectSpecified = Object.GetBoolField(TEXT("ProjectSpecified"));
 		ShareableProjectPath = Object.GetStringField(TEXT("ShareableProjectPath"));
-		CookConfiguration = (EBuildConfiguration)((int32)Object.GetNumberField(TEXT("CookConfiguration")));
-		CookOptions = Object.GetStringField(TEXT("CookOptions"));
-		CookMode = (TEnumAsByte<ELauncherProfileCookModes::Type>)((int32)Object.GetNumberField(TEXT("CookMode")));
-		CookUnversioned = Object.GetBoolField(TEXT("CookUnversioned"));
-
-		CookedCultures.Reset();
-		const TArray<TSharedPtr<FJsonValue>>* Cultures = NULL;
-		if (Object.TryGetArrayField(TEXT("CookedCultures"), Cultures))
-		{
-			for (auto Value : *Cultures)
-			{
-				CookedCultures.Add(Value->AsString());
-			}
-		}
-
-		CookedMaps.Reset();
-		const TArray<TSharedPtr<FJsonValue>>* Maps = NULL;
-		if (Object.TryGetArrayField(TEXT("CookedMaps"), Maps))
-		{
-			for (auto Value : *Maps)
-			{
-				CookedMaps.Add(Value->AsString());
-			}
-		}
-
-		CookedPlatforms.Reset();
-		const TArray<TSharedPtr<FJsonValue>>* Platforms = NULL;
-		if (Object.TryGetArrayField(TEXT("CookedPlatforms"), Platforms))
-		{
-			for (auto Value : *Platforms)
-			{
-				CookedPlatforms.Add(Value->AsString());
-			}
-		}
-
-		DeployStreamingServer = Object.GetBoolField(TEXT("DeployStreamingServer"));
-		DeployWithUnrealPak = Object.GetBoolField(TEXT("DeployWithUnrealPak"));
-		FGuid::Parse(Object.GetStringField(TEXT("DeployedDeviceGroupId")), DeployedDeviceGroupId);
-		DeploymentMode = (TEnumAsByte<ELauncherProfileDeploymentModes::Type>)((int32)Object.GetNumberField(TEXT("DeploymentMode")));
-		HideFileServerWindow = Object.GetBoolField(TEXT("HideFileServerWindow"));
-		LaunchMode = (TEnumAsByte<ELauncherProfileLaunchModes::Type>)((int32)Object.GetNumberField(TEXT("LaunchMode")));
-		PackagingMode = (TEnumAsByte<ELauncherProfilePackagingModes::Type>)((int32)Object.GetNumberField(TEXT("PackagingMode")));
-		PackageDir = Object.GetStringField(TEXT("PackageDir"));
-
-		int64 BuildModeValue;
-		if (Object.TryGetNumberField(TEXT("BuildMode"), BuildModeValue))
-		{
-			BuildMode = (TEnumAsByte<ELauncherProfileBuildModes::Type>)(int32)BuildModeValue;
-		}
-		else
-		{
-			BuildMode = Object.GetBoolField(TEXT("BuildGame")) ? ELauncherProfileBuildModes::Build : ELauncherProfileBuildModes::DoNotBuild;
-		}
-
-		ForceClose = Object.GetBoolField(TEXT("ForceClose"));
-		Timeout = (uint32)Object.GetNumberField(TEXT("Timeout"));
-		Compressed = Object.GetBoolField(TEXT("Compressed"));
-
-		if (Version >= LAUNCHERSERVICES_ADDEDENCRYPTINIFILES)
-		{
-			EncryptIniFiles = Object.GetBoolField(TEXT("EncryptIniFiles"));
-			ForDistribution = Object.GetBoolField(TEXT("ForDistribution"));
-		}
-		else
-		{
-			EncryptIniFiles = false;
-			ForDistribution = false;
-		}
-
-		DefaultDeployPlatform = *(Object.GetStringField(TEXT("DeployPlatform")));
-		bSkipCookingEditorContent = Object.GetBoolField(TEXT("SkipCookingEditorContent"));
-		DeployIncremental = Object.GetBoolField(TEXT("DeployIncremental"));
-		GeneratePatch = Object.GetBoolField(TEXT("GeneratePatch"));
-
-		if (Version >= LAUNCHERSERVICES_ADDEDMULTILEVELPATCHING)
-		{
-			AddPatchLevel = Object.GetBoolField(TEXT("AddPatchLevel"));
-			StageBaseReleasePaks = Object.GetBoolField(TEXT("StageBaseReleasePaks"));
-		}
-		else
-		{
-			AddPatchLevel = false;
-			StageBaseReleasePaks = false;
-		}
-
-		DLCIncludeEngineContent = Object.GetBoolField(TEXT("DLCIncludeEngineContent"));
-		CreateReleaseVersion = Object.GetBoolField(TEXT("CreateReleaseVersion"));
-		CreateReleaseVersionName = Object.GetStringField(TEXT("CreateReleaseVersionName"));
-		BasedOnReleaseVersionName = Object.GetStringField(TEXT("BasedOnReleaseVersionName"));
-
-		if (Version >= LAUNCHERSERVICES_ADDEDREFERENCECONTAINERS)
-		{
-			ReferenceContainerCryptoKeysFileName = Object.GetStringField(TEXT("ReferenceContainerCryptoKeysFileName"));
-			ReferenceContainerGlobalFileName = Object.GetStringField(TEXT("ReferenceContainerGlobalFileName"));
-		}
-		else
-		{
-			ReferenceContainerCryptoKeysFileName.Empty();
-			ReferenceContainerGlobalFileName.Empty();
-		}
-
-		if (Version >= LAUNCHERSERVICES_ADDEDORIGINALRELEASEVERSION)
-		{
-			OriginalReleaseVersionName = Object.GetStringField(TEXT("OriginalReleaseVersionName"));
-		}
-		else
-		{
-			OriginalReleaseVersionName.Empty();
-		}
-
-		CreateDLC = Object.GetBoolField(TEXT("CreateDLC"));
-		DLCName = Object.GetStringField(TEXT("DLCName"));
-		bGenerateChunks = Object.GetBoolField(TEXT("GenerateChunks"));
-		bGenerateHttpChunkData = Object.GetBoolField(TEXT("GenerateHttpChunkData"));
-		HttpChunkDataDirectory = Object.GetStringField(TEXT("HttpChunkDataDirectory"));
-		HttpChunkDataReleaseName = Object.GetStringField(TEXT("HttpChunkDataReleaseName"));
-
-		if (Version >= LAUNCHERSERVICES_ADDARCHIVE)
-		{
-			bArchive = Object.GetBoolField(TEXT("Archive"));
-			ArchiveDir = Object.GetStringField(TEXT("ArchiveDirectory"));
-		}
-		else
-		{
-			bArchive = false;
-			ArchiveDir = TEXT("");
-		}
-
-		if (Version >= LAUNCHERSERVICES_ADDEDADDITIONALCOMMANDLINE)
-		{
-			AdditionalCommandLineParameters = Object.GetStringField(TEXT("AdditionalCommandLineParameters"));
-		}
-		else
-		{
-			AdditionalCommandLineParameters = TEXT("");
-		}
-
-		if (Version >= LAUNCHERSERVICES_ADDEDINCLUDEPREREQUISITES)
-		{
-			IncludePrerequisites = Object.GetBoolField(TEXT("IncludePrerequisites"));
-		}
-
-		if (Version >= LAUNCHERSERVICES_ADDEDUSEIOSTORE)
-		{
-			bUseIoStore = Object.GetBoolField(TEXT("UseIoStore"));
-		}
-
-		if (Version >= LAUNCHERSERVICES_ADDEDMAKEBINARYCONFIG)
-		{
-			bMakeBinaryConfig = Object.GetBoolField(TEXT("MakeBinaryConfig"));
-		}
-
-		FString BuildTargetName;
-		if (Version >= LAUNCHERSERVICES_ADDBUILDTARGETNAME)
-		{
-			BuildTargetSpecified = Object.GetBoolField(TEXT("BuildTargetSpecified"));
-			if (Version < LAUNCHERSERVICES_ADDMULTIPLEBUILDTARGETS)
-			{
-				BuildTargetName = Object.GetStringField(TEXT("BuildTargetName"));
-				if (!BuildTargetName.IsEmpty())
-				{
-					BuildTargets.Add(BuildTargetName);
-				}
-			}
-		}
-
-		if (Version >= LAUNCHERSERVICES_ADDZENOPTIONS)
-		{
-			bUseZenStore = Object.GetBoolField(TEXT("UseZenStore"));
-			bImportingZenSnapshot = Object.GetBoolField(TEXT("ImportingZenSnapshot"));
-			bUseZenPakStreaming = Object.GetBoolField(TEXT("UseZenPakStreaming"));
-			ZenPakStreamingPath = Object.GetStringField(TEXT("ZenPakStreamingPath"));
-		}
-
-		if (Version >= LAUNCHERSERVICES_ADDBUILDARCHITECTURE)
-		{
-			ClientArchitectures.Reset();
-			ServerArchitectures.Reset();
-			EditorArchitectures.Reset();
-			Object.TryGetStringArrayField(TEXT("ClientArchitectures"), ClientArchitectures);
-			Object.TryGetStringArrayField(TEXT("ServerArchitectures"), ServerArchitectures);
-			Object.TryGetStringArrayField(TEXT("EditorArchitectures"), EditorArchitectures);
-		}
 
 		if (Version >= LAUNCHERSERVICES_ADDCUSTOMPROPERTIES)
 		{
@@ -2246,16 +1548,6 @@ PRAGMA_ENABLE_DEPRECATION_WARNINGS
 			}
 		}
 
-		if (Version >= LAUNCHERSERVICES_INCREMENTALCOOKMODE && Version < LAUNCHERSERVICES_REMOVEDINCREMENTALCOOKMODE)
-		{
-			ELauncherProfileIncrementalCookMode::Type IncrementalCookMode = (TEnumAsByte<ELauncherProfileIncrementalCookMode::Type>)((int32)Object.GetNumberField(TEXT("IncrementalCookMode")));
-			CookIncremental = (IncrementalCookMode != ELauncherProfileIncrementalCookMode::None);
-		}
-		else if (Version < LAUNCHERSERVICES_INCREMENTALCOOKMODE || Version >= LAUNCHERSERVICES_REMOVEDINCREMENTALCOOKMODE)
-		{
-			CookIncremental = Object.GetBoolField(TEXT("CookIncremental"));
-		}
-
 		if (Version >= LAUNCHERSERVICES_AUTOMATEDTESTS)
 		{
 			if (Version < LAUNCHERSERVICES_ADDUATCOMMANDS)
@@ -2265,7 +1557,16 @@ PRAGMA_ENABLE_DEPRECATION_WARNINGS
 				{
 					for (TSharedPtr<FJsonValue> Value : *Values)
 					{
-						UATCommands.Add(MakeShared<FLauncherProfileAutomatedTest>(*Value->AsObject().Get()));
+						ILauncherProfileUATCommandRef UATCommand = CreateAutomatedTestInternal();
+						if (UATCommand->Load((*Value->AsObject().Get())))
+						{
+							AddUATCommandInternal(UATCommand);
+						}
+						else
+						{
+							UE_LOG(LogLauncherProfile, Error, TEXT("cannot parse automated test"));
+							return false;
+						}
 					}
 				}
 			}
@@ -2274,44 +1575,6 @@ PRAGMA_ENABLE_DEPRECATION_WARNINGS
 			bWantAutomatedTestBuild = Object.GetBoolField(TEXT("WantAutomatedTestBuild"));
 		}
 
-		if (Version >= LAUNCHERSERVICES_FASTITERATE)
-		{
-			bFastIterate = Object.GetBoolField(TEXT("FastIterate"));
-		}
-
-		if (Version >= LAUNCHERSERVICES_ADDMULTIPLEBUILDTARGETS)
-		{
-			BuildTargets.Reset();
-			const TArray<TSharedPtr<FJsonValue>>* JsonValues = NULL;
-			if (Object.TryGetArrayField(TEXT("BuildTargets"), JsonValues))
-			{
-				for (auto Value : *JsonValues)
-				{
-					BuildTargets.Add(Value->AsString());
-				}
-			}
-
-			if (BuildTargets.Num() == 0 && !BuildTargetName.IsEmpty())
-			{
-				BuildTargets.Add(BuildTargetName);
-			}
-
-
-			AdditionalTargetCommandLineParameters.Reset();
-			const TSharedPtr<FJsonObject>* AdditionalTargetCommandLineParametersObject;
-			if (Object.TryGetObjectField(TEXT("AdditionalTargetCommandLineParameters"), AdditionalTargetCommandLineParametersObject))
-			{
-				for (const auto& KeyPairs : (*AdditionalTargetCommandLineParametersObject)->Values)
-				{
-					EBuildTargetType BuildTarget;
-					if (LexTryParseString(BuildTarget, *KeyPairs.Key))
-					{
-						AdditionalTargetCommandLineParameters.Add(BuildTarget, KeyPairs.Value->AsString());
-					}
-				}
-			}
-		}
-
 
 		if (Version >= LAUNCHERSERVICES_ADDUATCOMMANDS)
 		{
@@ -2322,37 +1585,29 @@ PRAGMA_ENABLE_DEPRECATION_WARNINGS
 				{
 					const FJsonObject& JsonObject = *Value->AsObject().Get();
 					FString UATCommandType = JsonObject.GetStringField(TEXT("Type"));
-					if (UATCommandType == FLauncherProfileUATCommand::TypeName)
+
+					ILauncherProfileUATCommandPtr UATCommand = CreateUATCommandByTypeNameInternal(*UATCommandType);
+					if (!UATCommand.IsValid())
 					{
-						UATCommands.Add(MakeShared<FLauncherProfileUATCommand>(JsonObject));
+						UE_LOG(LogLauncherProfile, Error, TEXT("unknown UAT command type: %s"), *UATCommandType);
+						return false;
 					}
-					else if (UATCommandType == FLauncherProfileAutomatedTest::TypeName)
+					else if (UATCommand->Load(JsonObject))
 					{
-						UATCommands.Add(MakeShared<FLauncherProfileAutomatedTest>(JsonObject));
+						AddUATCommandInternal(UATCommand.ToSharedRef());
 					}
 					else
 					{
-						UE_LOG(LogLauncherProfile, Error, TEXT("unknown UAT command type: %s"), *UATCommandType);
+						UE_LOG(LogLauncherProfile, Error, TEXT("cannot parse UAT command type: %s"), *UATCommandType);
 						return false;
 					}
 				}
 			}
-		}
-
 
-		// load the default launch role
-		TSharedPtr<FJsonObject> Role = Object.GetObjectField(TEXT("DefaultRole"));
-		DefaultLaunchRole->Load(*(Role.Get()));
-
-		// serialize the launch roles
-		DeployedDeviceGroup.Reset();
-		LaunchRoles.Reset();
-		const TArray<TSharedPtr<FJsonValue>>* Roles = NULL;
-		if (Object.TryGetArrayField(TEXT("LaunchRoles"), Roles))
-		{
-			for (auto Value : *Roles)
+			if (Version >= LAUNCHERSERVICES_INNERBUILDCOOKRUN)
 			{
-				LaunchRoles.Add(MakeShareable(new FLauncherProfileLaunchRole(*(Value->AsObject().Get()))));
+				TArray<ILauncherProfileBuildCookRunRef> BuildCookRunCommands = GetBuildCookRunCommands();
+				FirstBuildCookRun = (BuildCookRunCommands.Num() > 0) ? BuildCookRunCommands[0].ToSharedPtr() : nullptr;
 			}
 		}
 
@@ -2361,11 +1616,6 @@ PRAGMA_ENABLE_DEPRECATION_WARNINGS
 			FullProjectPath = FPaths::ConvertRelativePathToFull(FPaths::RootDir(), ShareableProjectPath);
 		}
 
-		if (DefaultDeployPlatform != NAME_None)
-		{
-			SetDefaultDeployPlatform(DefaultDeployPlatform);
-		}
-
 		RefreshValidBuildTargets();
 		Validate();
 
@@ -2375,7 +1625,6 @@ PRAGMA_ENABLE_DEPRECATION_WARNINGS
 	virtual void SetDefaults( ) override
 	{
 		ProjectSpecified = false;
-		AdditionalCommandLineParameters = FString();
 
 		// default project settings
 		if (FPaths::IsProjectFilePathSet())
@@ -2407,98 +1656,28 @@ PRAGMA_ENABLE_DEPRECATION_WARNINGS
 		// Use the locally specified project path is resolving through the root isn't working
 		ProjectSpecified = !FullProjectPath.IsEmpty();
 		
-		// I don't use FApp::GetBuildConfiguration() because i don't want the act of running in debug the first time to cause 
-		// profiles the user creates to be in debug. This will keep consistency.
-		BuildConfiguration = EBuildConfiguration::Development;
-
-		FInternationalization& I18N = FInternationalization::Get();
-
 		// default build settings
-		BuildMode = ELauncherProfileBuildModes::Auto;
-		bFastIterate = false;
 		BuildUAT = !FApp::GetEngineIsPromotedBuild() && !FApp::IsEngineInstalled();
-		BuildTargetSpecified = false;
 
-		// default cook settings
-		CookConfiguration = FApp::GetBuildConfiguration();
-		CookMode = ELauncherProfileCookModes::OnTheFly;
-		CookOptions = FString();
-		CookIncremental = false;
-		IterateSharedCookedBuild=false;
-		CookUnversioned = true;
-		Compressed = true;
-		EncryptIniFiles = false;
-		ForDistribution = false;
-		CookedCultures.Reset();
-		CookedCultures.Add(I18N.GetCurrentCulture()->GetName());
-		CookedMaps.Reset();
-		CookedPlatforms.Reset();
-		BuildTargets.Reset();
-		bSkipCookingEditorContent = false;
-        ForceClose = true;
-        Timeout = 60;
-
-/*		if (GetTargetPlatformManager()->GetRunningTargetPlatform() != NULL)
-		{
-			CookedPlatforms.Add(GetTargetPlatformManager()->GetRunningTargetPlatform()->PlatformName());
-		}*/	
-
-		bArchive = false;
-		ArchiveDir = TEXT("");
 
 		// default deploy settings
-		DeployedDeviceGroup.Reset();
-		DeploymentMode = ELauncherProfileDeploymentModes::CopyToDevice;
-		DeployStreamingServer = false;
-		IncludePrerequisites = false;
-		DeployWithUnrealPak = false;
-		DeployedDeviceGroupId = FGuid();
-		HideFileServerWindow = false;
-		DeployIncremental = false;
-
-		CreateReleaseVersion = false;
-		GeneratePatch = false;
-		AddPatchLevel = false;
-		StageBaseReleasePaks = false;
-		CreateDLC = false;
-		DLCIncludeEngineContent = false;
-
-		bGenerateChunks = false;
-		bGenerateHttpChunkData = false;
-		HttpChunkDataDirectory = TEXT("");
-		HttpChunkDataReleaseName = TEXT("");
-
-		// default launch settings
-		DefaultDeployPlatform = NAME_None;
-		LaunchMode = ELauncherProfileLaunchModes::DefaultRole;
-		DefaultLaunchRole->SetCommandLine(FString());
-		DefaultLaunchRole->SetInitialCulture(I18N.GetCurrentCulture()->GetName());
-		DefaultLaunchRole->SetInitialMap(FString());
-		DefaultLaunchRole->SetName(TEXT("Default Role"));
-		DefaultLaunchRole->SetInstanceType(ELauncherProfileRoleInstanceTypes::StandaloneClient);
-		DefaultLaunchRole->SetVsyncEnabled(false);
-		LaunchRoles.Reset();
-
-		// default packaging settings
-		PackagingMode = ELauncherProfilePackagingModes::DoNotPackage;
-
+		StageBaseReleasePaks = false; // deprecated in 5.6
 		// default UAT settings
 		ILauncherServicesModule& LauncherServicesModule = FModuleManager::GetModuleChecked<ILauncherServicesModule>(TEXT("LauncherServices"));
 		EditorExe = LauncherServicesModule.GetExecutableForCommandlets();
 
 		bNotForLicensees = false;
-		bUseIoStore = false;
-		bUseZenStore = false;
-		bImportingZenSnapshot = false;
-		bUseZenPakStreaming = false;
 		bShouldUpdateFlash = false;
-		bIsDeviceASimulator = false;
-		bMakeBinaryConfig = false;
 
 		AutomatedTestBuildPath.Reset();
 		bWantAutomatedTestBuild = false;
 
 		UATCommands.Reset();
+		UATCommandValidation.Reset();
+		FirstBuildCookRun = CreateBuildCookRunInternal(TEXT("FirstBuildCookRun"));
+		AddUATCommandInternal(FirstBuildCookRun.ToSharedRef());
+		FirstBuildCookRun->SetDeployedDeviceGroup(LauncherProfileManager->AddNewDeviceGroup());
+
 
 		RefreshValidBuildTargets();
 		Validate();
@@ -2506,17 +1685,14 @@ PRAGMA_ENABLE_DEPRECATION_WARNINGS
 
 	virtual void SetBuildMode(ELauncherProfileBuildModes::Type Mode) override
 	{
-		if (BuildMode != Mode)
-		{
-			BuildMode = Mode;
-
-			Validate();
-		}
+		check(FirstBuildCookRun.IsValid())
+		FirstBuildCookRun->SetBuildMode(Mode);
 	}
 
 	virtual void SetFastIterate(bool Enable) override
 	{
-		bFastIterate = Enable;
+		check(FirstBuildCookRun.IsValid())
+		FirstBuildCookRun->SetFastIterate(Enable);
 	}
 	
 	virtual void SetBuildUAT(bool Build) override
@@ -2531,53 +1707,28 @@ PRAGMA_ENABLE_DEPRECATION_WARNINGS
 
 	virtual void SetAdditionalCommandLineParameters(const FString& Params) override
 	{
-		if (AdditionalCommandLineParameters != Params)
-		{
-			AdditionalCommandLineParameters = Params;
-
-			Validate();
-		}
+		check(FirstBuildCookRun.IsValid())
+		FirstBuildCookRun->SetAdditionalCommandLineParameters(Params);
 	}
 
 	virtual void SetAdditionalTargetCommandLineParameters(const FString& Params, EBuildTargetType BuildTargetType ) override
 	{
-		check( BuildTargetType != EBuildTargetType::Unknown);
-
-		FString CurrentParams = GetAdditionalTargetCommandLineParameters(BuildTargetType);
-		if (CurrentParams != Params)
-		{
-			if (Params.IsEmpty())
-			{
-				AdditionalTargetCommandLineParameters.Remove(BuildTargetType);
-			}
-			else
-			{
-				AdditionalTargetCommandLineParameters.Add(BuildTargetType, Params);
-			}
-
-			Validate();
-		}
+		check(FirstBuildCookRun.IsValid())
+		FirstBuildCookRun->SetAdditionalTargetCommandLineParameters(Params, BuildTargetType);
 	}
 
 
 
 	virtual void SetBuildConfiguration( EBuildConfiguration Configuration ) override
 	{
-		if (BuildConfiguration != Configuration)
-		{
-			BuildConfiguration = Configuration;
-
-			Validate();
-		}
+		check(FirstBuildCookRun.IsValid())
+		FirstBuildCookRun->SetBuildConfiguration(Configuration);
 	}
 
 	virtual void SetBuildTargetSpecified(bool Specified) override
 	{
-		if (BuildTargetSpecified != Specified)
-		{
-			BuildTargetSpecified = Specified;
-			Validate();
-		}
+		check(FirstBuildCookRun.IsValid())
+		FirstBuildCookRun->SetBuildTargetSpecified(Specified);
 	}
 
 	virtual void FallbackBuildTargetUpdated() override
@@ -2591,16 +1742,16 @@ PRAGMA_ENABLE_DEPRECATION_WARNINGS
 
 	virtual void SetBuildTarget( const FString& TargetName ) override
 	{
+		check(FirstBuildCookRun.IsValid())
+		TArray<FString> BuildTargets = FirstBuildCookRun->GetBuildTargets();
+
 		FString BuildTargetName = (BuildTargets.Num() > 0) ? BuildTargets[0] : FString();
 		if (BuildTargetName != TargetName || BuildTargets.Num() > 1)
 		{
 			UE_CLOG(BuildTargets.Num() > 0, LogTemp, Warning, TEXT("SetBuildTarget removed multiple targets - legacy project launcher?"));
 
-			BuildTargets.Reset();
-			if (!TargetName.IsEmpty())
-			{
-				BuildTargets.Add(TargetName);
-			}
+			FirstBuildCookRun->ClearBuildTargets();
+			FirstBuildCookRun->AddBuildTarget(TargetName);
 
 			Validate();
 		}
@@ -2608,114 +1759,62 @@ PRAGMA_ENABLE_DEPRECATION_WARNINGS
 
 	virtual void SetCookConfiguration( EBuildConfiguration Configuration ) override
 	{
-		if (CookConfiguration != Configuration)
-		{
-			CookConfiguration = Configuration;
-
-			Validate();
-		}
+		check(FirstBuildCookRun.IsValid())
+		FirstBuildCookRun->SetCookConfiguration(Configuration);
 	}
 
 	virtual void SetCookMode( ELauncherProfileCookModes::Type Mode ) override
 	{
-		if (CookMode != Mode)
-		{
-			CookMode = Mode;
-
-			Validate();
-		}
+		check(FirstBuildCookRun.IsValid())
+		FirstBuildCookRun->SetCookMode(Mode);
 	}
 
 	virtual void SetCookOptions(const FString& Options) override
 	{
-		if (CookOptions != Options)
-		{
-			CookOptions = Options;
-
-			Validate();
-		}
+		check(FirstBuildCookRun.IsValid())
+		FirstBuildCookRun->SetCookOptions(Options);
 	}
 
 	virtual void SetSkipCookingEditorContent(const bool InSkipCookingEditorContent) override
 	{
-		if (bSkipCookingEditorContent != InSkipCookingEditorContent)
-		{
-			bSkipCookingEditorContent = InSkipCookingEditorContent;
-			Validate();
-		}
+		check(FirstBuildCookRun.IsValid())
+		FirstBuildCookRun->SetSkipCookingEditorContent(InSkipCookingEditorContent);
 	}
 
 	virtual void SetDeployWithUnrealPak( bool UseUnrealPak ) override
 	{
-		if (DeployWithUnrealPak != UseUnrealPak)
-		{
-			DeployWithUnrealPak = UseUnrealPak;
-
-			Validate();
-		}
+		check(FirstBuildCookRun.IsValid())
+		FirstBuildCookRun->SetDeployWithUnrealPak(UseUnrealPak);
 	}
 
 	virtual void SetGenerateChunks(bool bInGenerateChunks) override
 	{
-		if (bGenerateChunks != bInGenerateChunks)
-		{
-			bGenerateChunks = bInGenerateChunks;
-			Validate();
-		}
+		check(FirstBuildCookRun.IsValid())
+		FirstBuildCookRun->SetGenerateChunks(bInGenerateChunks);
 	}
 
 	virtual void SetGenerateHttpChunkData(bool bInGenerateHttpChunkData) override
 	{
-		if (bGenerateHttpChunkData != bInGenerateHttpChunkData)
-		{
-			bGenerateHttpChunkData = bInGenerateHttpChunkData;
-			Validate();
-		}
+		check(FirstBuildCookRun.IsValid())
+		FirstBuildCookRun->SetGenerateHttpChunkData(bInGenerateHttpChunkData);
 	}
 
 	virtual void SetHttpChunkDataDirectory(const FString& InHttpChunkDataDirectory) override
 	{
-		if (HttpChunkDataDirectory != InHttpChunkDataDirectory)
-		{
-			HttpChunkDataDirectory = InHttpChunkDataDirectory;
-			Validate();
-		}
+		check(FirstBuildCookRun.IsValid())
+		FirstBuildCookRun->SetHttpChunkDataDirectory(InHttpChunkDataDirectory);
 	}
 
 	virtual void SetHttpChunkDataReleaseName(const FString& InHttpChunkDataReleaseName) override
 	{
-		if (HttpChunkDataReleaseName != InHttpChunkDataReleaseName)
-		{
-			HttpChunkDataReleaseName = InHttpChunkDataReleaseName;
-			Validate();
-		}
+		check(FirstBuildCookRun.IsValid())
+		FirstBuildCookRun->SetHttpChunkDataReleaseName(InHttpChunkDataReleaseName);
 	}
 
 	virtual void SetDeployedDeviceGroup( const ILauncherDeviceGroupPtr& DeviceGroup ) override
 	{
-		if(DeployedDeviceGroup.IsValid())
-		{
-			DeployedDeviceGroup->OnDeviceAdded().Remove(OnLauncherDeviceGroupDeviceAddedDelegateHandle);
-			DeployedDeviceGroup->OnDeviceRemoved().Remove(OnLauncherDeviceGroupDeviceRemoveDelegateHandle);
-		}
-		DeployedDeviceGroup = DeviceGroup;
-		if (DeployedDeviceGroup.IsValid())
-		{
-			OnLauncherDeviceGroupDeviceAddedDelegateHandle   = DeployedDeviceGroup->OnDeviceAdded().AddRaw(this, &FLauncherProfile::OnLauncherDeviceGroupDeviceAdded);
-			OnLauncherDeviceGroupDeviceRemoveDelegateHandle  = DeployedDeviceGroup->OnDeviceRemoved().AddRaw(this, &FLauncherProfile::OnLauncherDeviceGroupDeviceRemove);
-			DeployedDeviceGroupId = DeployedDeviceGroup->GetId();
-		}
-		else
-		{
-			DeployedDeviceGroupId = FGuid();
-		}
-
-		if (DefaultDeployPlatform != NAME_None)
-		{
-			SetDefaultDeployPlatform(DefaultDeployPlatform);
-		}
-
-		Validate();
+		check(FirstBuildCookRun.IsValid())
+		FirstBuildCookRun->SetDeployedDeviceGroup(DeviceGroup);
 	}
 
 	virtual FIsCookFinishedDelegate& OnIsCookFinished() override
@@ -2730,102 +1829,70 @@ PRAGMA_ENABLE_DEPRECATION_WARNINGS
 
 	virtual void SetDeploymentMode( ELauncherProfileDeploymentModes::Type Mode ) override
 	{
-		if (DeploymentMode != Mode)
-		{
-			DeploymentMode = Mode;
-
-			Validate();
-		}
+		check(FirstBuildCookRun.IsValid())
+		FirstBuildCookRun->SetDeploymentMode(Mode);
 	}
 
     virtual void SetForceClose( bool Close ) override
     {
-        if (ForceClose != Close)
-        {
-            ForceClose = Close;
-            Validate();
-        }
+		check(FirstBuildCookRun.IsValid())
+        FirstBuildCookRun->SetForceClose(Close);
     }
     
 	virtual void SetHideFileServerWindow( bool Hide ) override
 	{
-		HideFileServerWindow = Hide;
+		check(FirstBuildCookRun.IsValid())
+		FirstBuildCookRun->SetHideFileServerWindow(Hide);
 	}
 
 	virtual void SetIncrementalCooking( bool Incremental ) override
 	{
-		if (Incremental != CookIncremental)
-		{
-			CookIncremental = Incremental;
-			Validate();
-		}
+		check(FirstBuildCookRun.IsValid())
+		FirstBuildCookRun->SetIncrementalCooking(Incremental);
 	}
 
 	virtual void SetIncrementalCookMode( ELauncherProfileIncrementalCookMode::Type Mode ) override
 	{
+		check(FirstBuildCookRun.IsValid())
 		bool Incremental = (Mode != ELauncherProfileIncrementalCookMode::None);
-		SetIncrementalCooking(Incremental);
+		FirstBuildCookRun->SetIncrementalCooking(Incremental);
 	}
 
 
 	virtual void SetIterateSharedCookedBuild( bool SharedCookedBuild ) override
 	{
-		if (IterateSharedCookedBuild != SharedCookedBuild)
-		{
-			IterateSharedCookedBuild = SharedCookedBuild;
-
-			Validate();
-		}
+		check(FirstBuildCookRun.IsValid())
+		FirstBuildCookRun->SetIterateSharedCookedBuild(SharedCookedBuild);
 	}
 
 	virtual void SetCompressed( bool Enabled ) override
-	{
-		if (Compressed != Enabled)
-		{
-			Compressed = Enabled;
-
-			Validate();
-		}
+	{
+		check(FirstBuildCookRun.IsValid())
+		FirstBuildCookRun->SetCompressed(Enabled);
 	}
 
 	virtual void SetForDistribution(bool Enabled)override
 	{
-		if (ForDistribution != Enabled)
-		{
-			ForDistribution = Enabled;
-
-			Validate();
-		}
+		check(FirstBuildCookRun.IsValid())
+		FirstBuildCookRun->SetForDistribution(Enabled);
 	}
 
 	virtual void SetEncryptingIniFiles(bool Enabled) override
 	{
-		if (EncryptIniFiles != Enabled)
-		{
-			EncryptIniFiles = Enabled;
-
-			Validate();
-		}
+		check(FirstBuildCookRun.IsValid())
+		FirstBuildCookRun->SetEncryptingIniFiles(Enabled);
 	}
 
 	virtual void SetIncrementalDeploying( bool Incremental ) override
 	{
-		if (DeployIncremental != Incremental)
-		{
-			DeployIncremental = Incremental;
-
-			Validate();
-		}
+		check(FirstBuildCookRun.IsValid())
+		FirstBuildCookRun->SetIncrementalDeploying(Incremental);
 	}
 
 	virtual void SetLaunchMode( ELauncherProfileLaunchModes::Type Mode ) override
 	{
-		if (LaunchMode != Mode)
-		{
-			LaunchMode = Mode;
-
-			Validate();
-		}
+		check(FirstBuildCookRun.IsValid())
+		FirstBuildCookRun->SetLaunchMode(Mode);
 	}
 
 	virtual void SetName( const FString& NewName ) override
@@ -2855,42 +1922,26 @@ PRAGMA_ENABLE_DEPRECATION_WARNINGS
 
 	virtual void SetPackagingMode( ELauncherProfilePackagingModes::Type Mode ) override
 	{
-		if (PackagingMode != Mode)
-		{
-			PackagingMode = Mode;
-
-			Validate();
-		}
+		check(FirstBuildCookRun.IsValid())
+		FirstBuildCookRun->SetPackagingMode(Mode);
 	}
 
 	virtual void SetPackageDirectory( const FString& Dir ) override
 	{
-		if (PackageDir != Dir)
-		{
-			PackageDir = Dir;
-
-			Validate();
-		}
+		check(FirstBuildCookRun.IsValid())
+		FirstBuildCookRun->SetPackageDirectory(Dir);
 	}
 
 	virtual void SetArchive( bool bInArchive ) override
 	{
-		if (bInArchive != bArchive)
-		{
-			bArchive = bInArchive;
-
-			Validate();
-		}
+		check(FirstBuildCookRun.IsValid())
+		FirstBuildCookRun->SetArchive(bInArchive);
 	}
 
 	virtual void SetArchiveDirectory( const FString& Dir ) override
 	{
-		if (ArchiveDir != Dir)
-		{
-			ArchiveDir = Dir;
-
-			Validate();
-		}
+		check(FirstBuildCookRun.IsValid())
+		FirstBuildCookRun->SetArchiveDirectory(Dir);
 	}
 
 	virtual void SetProjectSpecified(bool Specified) override
@@ -2942,7 +1993,11 @@ PRAGMA_ENABLE_DEPRECATION_WARNINGS
 					ShareableProjectPath = FullProjectPath;
 				}				
 			}
-			CookedMaps.Reset();
+
+			for (const ILauncherProfileBuildCookRunRef& BuildCookRun : GetBuildCookRunCommands())
+			{
+				BuildCookRun->ClearCookedMaps();
+			}
 
 			RefreshValidBuildTargets();
 			Validate();
@@ -2953,52 +2008,38 @@ PRAGMA_ENABLE_DEPRECATION_WARNINGS
 
 	virtual void SetStreamingFileServer( bool Streaming ) override
 	{
-		if (DeployStreamingServer != Streaming)
-		{
-			DeployStreamingServer = Streaming;
-
-			Validate();
-		}
+		check(FirstBuildCookRun.IsValid())
+		FirstBuildCookRun->SetStreamingFileServer(Streaming);
 	}
 
 	virtual void SetIncludePrerequisites(bool InValue) override
 	{
-		if (IncludePrerequisites != InValue)
-		{
-			IncludePrerequisites = InValue;
-
-			Validate();
-		}
+		check(FirstBuildCookRun.IsValid())
+		FirstBuildCookRun->SetIncludePrerequisites(InValue);
 	}
 
     virtual void SetTimeout( uint32 InTime ) override
     {
-        if (Timeout != InTime)
-        {
-            Timeout = InTime;
-            
-            Validate();
-        }
+		check(FirstBuildCookRun.IsValid())
+        FirstBuildCookRun->SetTimeout(InTime);
     }
     
 	virtual void SetUnversionedCooking( bool Unversioned ) override
 	{
-		if (CookUnversioned != Unversioned)
-		{
-			CookUnversioned = Unversioned;
-
-			Validate();
-		}
+		check(FirstBuildCookRun.IsValid())
+		FirstBuildCookRun->SetUnversionedCooking(Unversioned);
 	}
 
 	virtual void SetGeneratePatch( bool InGeneratePatch ) override
 	{
-		GeneratePatch = InGeneratePatch;
+		check(FirstBuildCookRun.IsValid())
+		FirstBuildCookRun->SetGeneratePatch(InGeneratePatch);
 	}
 
 	virtual void SetAddPatchLevel( bool InAddPatchLevel) override
 	{
-		AddPatchLevel = InAddPatchLevel;
+		check(FirstBuildCookRun.IsValid())
+		FirstBuildCookRun->SetAddPatchLevel(InAddPatchLevel);
 	}
 
 	virtual void SetStageBaseReleasePaks(bool InStageBaseReleasePaks) override
@@ -3026,6 +2067,11 @@ PRAGMA_ENABLE_DEPRECATION_WARNINGS
 		return CustomValidationDelegate;
 	}
 
+	virtual FOnProfileCustomUATCommandValidation& OnCustomUATCommandValidation() override
+	{
+		return CustomUATCommandValidationDelegate;
+	}
+
 	virtual void SetEditorExe( const FString& InEditorExe ) override
 	{
 		EditorExe = InEditorExe;
@@ -3038,82 +2084,64 @@ PRAGMA_ENABLE_DEPRECATION_WARNINGS
 
 	virtual void SetUseIoStore(bool bInUseIoStore) override
 	{
-		bUseIoStore = bInUseIoStore;
-
-		if (bUseIoStore)
-		{
-			SetDeployWithUnrealPak(true);
-		}
-
-		Validate();
+		check(FirstBuildCookRun.IsValid())
+		FirstBuildCookRun->SetUseIoStore(bInUseIoStore);
 	}
 
 	virtual bool IsUsingIoStore() const override
 	{
-		return bUseIoStore;
+		check(FirstBuildCookRun.IsValid())
+		return FirstBuildCookRun->IsUsingIoStore();
 	}
 
 	virtual void SetUseZenStore(bool bInUseZenStore) override
 	{
-		bUseZenStore = bInUseZenStore;
-
-		Validate();
+		check(FirstBuildCookRun.IsValid())
+		FirstBuildCookRun->SetUseZenStore(bInUseZenStore);
 	}
 
 	virtual bool IsUsingZenStore() const override
 	{
-		return bUseZenStore;
+		check(FirstBuildCookRun.IsValid())
+		return FirstBuildCookRun->IsUsingZenStore();
 	}
 
 	virtual int32 GetZenSnapshot() override;
 
 	virtual void SetImportingZenSnapshot( bool Import ) override
 	{
-		if (bImportingZenSnapshot != Import)
-		{
-			bImportingZenSnapshot = Import;
-
-			Validate();
-		}
+		check(FirstBuildCookRun.IsValid())
+		FirstBuildCookRun->SetImportingZenSnapshot(Import);
 	}
 
 	virtual bool IsImportingZenSnapshot() const override
 	{
-		return bImportingZenSnapshot;
+		check(FirstBuildCookRun.IsValid())
+		return FirstBuildCookRun->IsImportingZenSnapshot();
 	}
 
 	virtual void SetUseZenPakStreaming( bool UseZenPakStreaming ) override
 	{
-		if (bUseZenPakStreaming != UseZenPakStreaming)
-		{
-			bUseZenPakStreaming = UseZenPakStreaming;
-
-			Validate();
-		}
+		check(FirstBuildCookRun.IsValid())
+		FirstBuildCookRun->SetUseZenPakStreaming(UseZenPakStreaming);
 	}
 
 	virtual bool IsUsingZenPakStreaming() const override
 	{
-		if (CookMode == ELauncherProfileCookModes::DoNotCook)
-		{
-			return bUseZenPakStreaming;
-		}
-		return false;
+		check(FirstBuildCookRun.IsValid())
+		return FirstBuildCookRun->IsUsingZenPakStreaming();
 	}
 
 	virtual void SetZenPakStreamingPath( const FString& Path ) override
 	{
-		if (ZenPakStreamingPath != Path)
-		{
-			ZenPakStreamingPath = Path;
-
-			Validate();
-		}
+		check(FirstBuildCookRun.IsValid())
+		FirstBuildCookRun->SetZenPakStreamingPath(Path);
 	}
 
 	virtual FString GetZenPakStreamingPath() const override
 	{
-		return ZenPakStreamingPath;
+		check(FirstBuildCookRun.IsValid())
+		return FirstBuildCookRun->GetZenPakStreamingPath();
 	}
 
 	virtual void SetShouldUpdateDeviceFlash(bool bInShouldUpdateFlash) override
@@ -3131,7 +2159,8 @@ PRAGMA_ENABLE_DEPRECATION_WARNINGS
 
 	virtual void SetDeviceIsASimulator(bool bInIsDeviceASimualtor) override
 	{
-		bIsDeviceASimulator = bInIsDeviceASimualtor;
+		check(FirstBuildCookRun.IsValid())
+		FirstBuildCookRun->SetDeviceIsASimulator(bInIsDeviceASimualtor);
 	}
 
 	/**
@@ -3139,17 +2168,20 @@ PRAGMA_ENABLE_DEPRECATION_WARNINGS
 	 */
 	virtual bool IsDeviceASimulator() const override
 	{
-		return bIsDeviceASimulator;
+		check(FirstBuildCookRun.IsValid())
+		return FirstBuildCookRun->IsDeviceASimulator();
 	}
 
 	virtual void SetMakeBinaryConfig(bool bInMakeBinaryConfig) override
 	{
-		bMakeBinaryConfig = bInMakeBinaryConfig;
+		check(FirstBuildCookRun.IsValid())
+		FirstBuildCookRun->SetMakeBinaryConfig(bInMakeBinaryConfig);
 	}
 
 	virtual bool MakeBinaryConfig() const override
 	{
-		return bMakeBinaryConfig;
+		check(FirstBuildCookRun.IsValid())
+		return FirstBuildCookRun->MakeBinaryConfig();
 	}
 
 	virtual TArray<FString> GetExplicitBuildTargetNames() const override
@@ -3164,32 +2196,38 @@ PRAGMA_ENABLE_DEPRECATION_WARNINGS
 
 	virtual void SetClientArchitectures( const TArray<FString>& InArchitectures )
 	{
-		ClientArchitectures = InArchitectures;
+		check(FirstBuildCookRun.IsValid())
+		FirstBuildCookRun->SetClientArchitectures(InArchitectures);
 	}
 
 	virtual const TArray<FString>& GetClientArchitectures() const override
 	{
-		return ClientArchitectures;
+		check(FirstBuildCookRun.IsValid())
+		return FirstBuildCookRun->GetClientArchitectures();
 	}
 
 	virtual void SetServerArchitectures( const TArray<FString>& InArchitectures )
 	{
-		ServerArchitectures = InArchitectures;
+		check(FirstBuildCookRun.IsValid())
+		FirstBuildCookRun->SetServerArchitectures(InArchitectures);
 	}
 
 	virtual const TArray<FString>& GetServerArchitectures() const override
 	{
-		return ServerArchitectures;
+		check(FirstBuildCookRun.IsValid())
+		return FirstBuildCookRun->GetServerArchitectures();
 	}
 
 	virtual void SetEditorArchitectures( const TArray<FString>& InArchitectures )
 	{
-		EditorArchitectures = InArchitectures;
+		check(FirstBuildCookRun.IsValid())
+		FirstBuildCookRun->SetEditorArchitectures(InArchitectures);
 	}
 
 	virtual const TArray<FString>& GetEditorArchitectures() const override
 	{
-		return EditorArchitectures;
+		check(FirstBuildCookRun.IsValid())
+		return FirstBuildCookRun->GetEditorArchitectures();
 	}
 
 	virtual TMap<FString,FString>& GetCustomStringProperties() override
@@ -3214,7 +2252,7 @@ PRAGMA_ENABLE_DEPRECATION_WARNINGS
 			}
 		}
 
-		return AutomatedTests;
+		return MoveTemp(AutomatedTests);
 	}
 
 	virtual ILauncherProfileAutomatedTestPtr GetAutomatedTest( const TCHAR* InternalName ) const override
@@ -3228,8 +2266,9 @@ PRAGMA_ENABLE_DEPRECATION_WARNINGS
 		ILauncherProfileAutomatedTestPtr AutomatedTest = GetAutomatedTest(InternalName);
 		if (AutomatedTest == nullptr)
 		{
-			AutomatedTest = MakeShared<FLauncherProfileAutomatedTest>(InternalName);
-			UATCommands.Add(AutomatedTest.ToSharedRef());
+			AutomatedTest = CreateAutomatedTestInternal(InternalName);
+			AddUATCommandInternal(AutomatedTest.ToSharedRef());
+			UATCommandAdded.Broadcast(AutomatedTest.ToSharedRef());
 		}
 
 		return AutomatedTest.ToSharedRef();
@@ -3259,6 +2298,62 @@ PRAGMA_ENABLE_DEPRECATION_WARNINGS
 	}
 
 
+
+
+
+	virtual const TArray<ILauncherProfileBuildCookRunRef> GetBuildCookRunCommands() const override
+	{
+		TArray<ILauncherProfileBuildCookRunRef> BuildCookRuns;
+		for (const ILauncherProfileUATCommandRef& UATCommand : UATCommands)
+		{
+			ILauncherProfileBuildCookRunPtr BuildCookRun = UATCommand->AsBuildCookRun();
+			if (BuildCookRun)
+			{
+				BuildCookRuns.Add(BuildCookRun.ToSharedRef());
+			}
+		}
+
+		return MoveTemp(BuildCookRuns);
+
+	}
+
+	virtual ILauncherProfileBuildCookRunPtr GetBuildCookRunCommand( const TCHAR* InternalName ) const override
+	{
+		ILauncherProfileUATCommandPtr UATCommand = GetUATCommand(InternalName);
+		return UATCommand ? UATCommand->AsBuildCookRun() : nullptr;
+	}
+
+	virtual ILauncherProfileBuildCookRunRef FindOrAddBuildCookRunCommand( const TCHAR* InternalName ) override
+	{
+		// if no name has been given, autogenerate one
+		FString UniqueName;
+		if (InternalName == nullptr)
+		{
+			UniqueName = GetUniqueUATCommandName(TEXT("CustomBuildCookRun"));
+			InternalName = UniqueName.GetCharArray().GetData();
+		}
+
+
+		ILauncherProfileBuildCookRunPtr BuildCookRun = GetBuildCookRunCommand(InternalName);
+		if (BuildCookRun == nullptr)
+		{
+			BuildCookRun = CreateBuildCookRunInternal(InternalName);
+			BuildCookRun->SetDeployedDeviceGroup(LauncherProfileManager->AddNewDeviceGroup());
+
+			if (!FirstBuildCookRun.IsValid())
+			{
+				FirstBuildCookRun = BuildCookRun;
+			}
+
+			AddUATCommandInternal(BuildCookRun.ToSharedRef());
+			UATCommandAdded.Broadcast(BuildCookRun.ToSharedRef());
+		}
+
+		return BuildCookRun.ToSharedRef();
+	}
+
+
+
 	virtual const TArray<ILauncherProfileUATCommandRef>& GetUATCommands() const override
 	{
 		return UATCommands;
@@ -3275,27 +2370,48 @@ PRAGMA_ENABLE_DEPRECATION_WARNINGS
 
 	}
 
-	virtual ILauncherProfileUATCommandRef FindOrAddUATCommand( const TCHAR* InternalName )
+	virtual ILauncherProfileUATCommandRef FindOrAddUATCommand( const TCHAR* InternalName ) override
 	{
 		ILauncherProfileUATCommandPtr UATCommand = GetUATCommand(InternalName);
 		if (UATCommand == nullptr)
 		{
-			UATCommand = MakeShared<FLauncherProfileUATCommand>(InternalName);
-			UATCommands.Add(UATCommand.ToSharedRef());
+			UATCommand = CreateUATCommandInternal(InternalName);
+
+			AddUATCommandInternal(UATCommand.ToSharedRef());
+			UATCommandAdded.Broadcast(UATCommand.ToSharedRef());
+
 		}
 
 		return UATCommand.ToSharedRef();
 	}
 
-	virtual void RemoveUATCommand( const TCHAR* InternalName )
+	virtual void RemoveUATCommand( const TCHAR* InternalName ) override
 	{
 		ILauncherProfileUATCommandPtr UATCommand = GetUATCommand(InternalName);
 		if (UATCommand != nullptr)
 		{
+			UATCommandRemoved.Broadcast(UATCommand.ToSharedRef());
+			UATCommandValidation.Remove(UATCommand);
 			UATCommands.Remove(UATCommand.ToSharedRef());
+
+			// re-cache the first build cook run command if necessary
+			if (UATCommand == FirstBuildCookRun)
+			{
+				TArray<ILauncherProfileBuildCookRunRef> BuildCookRunCommands = GetBuildCookRunCommands();
+				FirstBuildCookRun = (BuildCookRunCommands.Num() > 0) ? BuildCookRunCommands[0].ToSharedPtr() : nullptr;
+			}
 		}
 	}
 
+	virtual FOnLauncherProfileUATCommandChanged& OnUATCommandAdded() override
+	{
+		return UATCommandAdded;
+	}
+
+	virtual FOnLauncherProfileUATCommandChanged& OnUATCommandRemoved() override
+	{
+		return UATCommandRemoved;
+	}
 
 
 
@@ -3304,37 +2420,40 @@ PRAGMA_ENABLE_DEPRECATION_WARNINGS
 
 	virtual TArray<FString> GetBuildTargets() const
 	{
-		if (BuildTargetSpecified)
-		{
-			return BuildTargets;
-		}
-		else
-		{
-			return TArray<FString>{LauncherProfileManager->GetBuildTarget()};
-		}
+		check(FirstBuildCookRun.IsValid())
+		return FirstBuildCookRun->GetBuildTargets();
 	}
 
 	virtual void AddBuildTarget( const FString& InBuildTarget )
 	{
-		BuildTargets.AddUnique(InBuildTarget);
+		check(FirstBuildCookRun.IsValid())
+		FirstBuildCookRun->AddBuildTarget(InBuildTarget);
 		Validate();
 	}
 
 	virtual void RemoveBuildTarget( const FString& InBuildTarget )
 	{
-		BuildTargets.Remove(InBuildTarget);
+		check(FirstBuildCookRun.IsValid())
+		FirstBuildCookRun->RemoveBuildTarget(InBuildTarget);
 		Validate();
 	}
 
 	virtual void ClearBuildTargets()
 	{
-		BuildTargets.Reset();
+		check(FirstBuildCookRun.IsValid())
+		FirstBuildCookRun->ClearBuildTargets();
 		Validate();
 	}
 
+
 	virtual bool SupportsLegacyProjectLauncher() const override
 	{
-		return (UATCommands.Num() == 0);
+		return FirstBuildCookRun.IsValid() && UATCommands.Num() == 1;
+	}
+
+	virtual ILauncherProfileBuildCookRunPtr GetFirstBuildCookRun() const override
+	{
+		return FirstBuildCookRun;
 	}
 
 
@@ -3351,70 +2470,99 @@ protected:
 	 */
 	void Validate( )
 	{
-		ValidationErrors.Reset();
+		ProfileValidation.Reset();
+		RefreshCustomWarningsAndErrors(nullptr);
 
-		// Build: a build configuration must be selected
-		if (BuildConfiguration == EBuildConfiguration::Unknown)
+		// Build: a project must be selected
+		if (GetProjectPath().IsEmpty())
 		{
-			ValidationErrors.Add(ELauncherProfileValidationErrors::NoBuildConfigurationSelected);
+			ProfileValidation.Errors.Add(ELauncherProfileValidationErrors::NoProjectSelected);
 		}
 
-		// Build: a project must be selected
-		if (GetProjectPath().IsEmpty())
+		for (const ILauncherProfileBuildCookRunRef& BuildCookRun : GetBuildCookRunCommands())
+		{
+			Validate(BuildCookRun);
+		}
+	}
+
+	void Validate(const ILauncherProfileBuildCookRunRef& BuildCookRun) 
+	{
+		if (!UATCommandValidation.Contains(BuildCookRun.ToSharedPtr()))
+		{
+			return;
+		}
+
+		FValidationData& ValidationData = UATCommandValidation.FindChecked(BuildCookRun.ToSharedPtr());
+		ValidationData.Reset();
+
+		// don't verify if this build command is disabled - it won't affect the build
+		if (!BuildCookRun->IsEnabled())
+		{
+			return;
+		}
+
+		// temp. fix : merge in the global validation errors so legacy PL can display them
+		if (BuildCookRun == FirstBuildCookRun)
+		{
+			ValidationData.Errors = ProfileValidation.Errors;
+		}
+
+		// Build: a build configuration must be selected
+		if (BuildCookRun->GetBuildConfiguration() == EBuildConfiguration::Unknown)
 		{
-			ValidationErrors.Add(ELauncherProfileValidationErrors::NoProjectSelected);
+			ValidationData.Errors.Add(ELauncherProfileValidationErrors::NoBuildConfigurationSelected);
 		}
 
 		// Cook: at least one platform must be selected when cooking by the book
-		if ((CookMode == ELauncherProfileCookModes::ByTheBook) && (CookedPlatforms.Num() == 0))
+		if ((BuildCookRun->GetCookMode() == ELauncherProfileCookModes::ByTheBook) && (BuildCookRun->GetCookedPlatforms().Num() == 0))
 		{
-			ValidationErrors.Add(ELauncherProfileValidationErrors::NoPlatformSelected);
+			ValidationData.Errors.Add(ELauncherProfileValidationErrors::NoPlatformSelected);
 		}
 
 		// Cook: at least one culture must be selected when cooking by the book
-		if ((CookMode == ELauncherProfileCookModes::ByTheBook) && (CookedCultures.Num() == 0))
+		if ((BuildCookRun->GetCookMode() == ELauncherProfileCookModes::ByTheBook) && (BuildCookRun->GetCookedCultures().Num() == 0))
 		{
-			ValidationErrors.Add(ELauncherProfileValidationErrors::NoCookedCulturesSelected);
+			ValidationData.Errors.Add(ELauncherProfileValidationErrors::NoCookedCulturesSelected);
 		}
 
 		// Deploy: a device group must be selected when deploying builds
-		if ((DeploymentMode == ELauncherProfileDeploymentModes::CopyToDevice) && !DeployedDeviceGroupId.IsValid())
+		if ((BuildCookRun->GetDeploymentMode() == ELauncherProfileDeploymentModes::CopyToDevice) && !BuildCookRun->GetDeployedDeviceGroup(false).IsValid())
 		{
-			ValidationErrors.Add(ELauncherProfileValidationErrors::DeployedDeviceGroupRequired);
+			ValidationData.Errors.Add(ELauncherProfileValidationErrors::DeployedDeviceGroupRequired);
 		}
 
 		// Deploy: deployment by copying to devices requires cooking by the book unless the data is available from zen
-		if ((DeploymentMode == ELauncherProfileDeploymentModes::CopyToDevice) && ((CookMode != ELauncherProfileCookModes::ByTheBook)&&(CookMode!=ELauncherProfileCookModes::ByTheBookInEditor)) && !IsUsingZenPakStreaming() && !IsImportingZenSnapshot())
+		if ((BuildCookRun->GetDeploymentMode() == ELauncherProfileDeploymentModes::CopyToDevice) && ((BuildCookRun->GetCookMode() != ELauncherProfileCookModes::ByTheBook)&&(BuildCookRun->GetCookMode()!=ELauncherProfileCookModes::ByTheBookInEditor)) && !BuildCookRun->IsUsingZenPakStreaming() && !BuildCookRun->IsImportingZenSnapshot())
 		{
-			ValidationErrors.Add(ELauncherProfileValidationErrors::CopyToDeviceRequiresCookByTheBook);
+			ValidationData.Errors.Add(ELauncherProfileValidationErrors::CopyToDeviceRequiresCookByTheBook);
 		}
 
 		// Deploy: deployment by copying to devices requires no packaging
-		if ((DeploymentMode == ELauncherProfileDeploymentModes::CopyToDevice) && (PackagingMode != ELauncherProfilePackagingModes::DoNotPackage))
+		if ((BuildCookRun->GetDeploymentMode() == ELauncherProfileDeploymentModes::CopyToDevice) && (BuildCookRun->GetPackagingMode() != ELauncherProfilePackagingModes::DoNotPackage))
 		{
-			ValidationErrors.Add(ELauncherProfileValidationErrors::CopyToDeviceRequiresNoPackaging);
+			ValidationData.Errors.Add(ELauncherProfileValidationErrors::CopyToDeviceRequiresNoPackaging);
 		}
 
 		// Deploy: deployment by copying a packaged build to devices requires a package dir
-		if ((DeploymentMode == ELauncherProfileDeploymentModes::CopyRepository) && (PackageDir == TEXT("")))
+		if ((BuildCookRun->GetDeploymentMode() == ELauncherProfileDeploymentModes::CopyRepository) && (BuildCookRun->GetPackageDirectory() == TEXT("")))
 		{
-			ValidationErrors.Add(ELauncherProfileValidationErrors::NoPackageDirectorySpecified);
+			ValidationData.Errors.Add(ELauncherProfileValidationErrors::NoPackageDirectorySpecified);
 		}
 
 		// Launch: custom launch roles are not supported yet
-		if (LaunchMode == ELauncherProfileLaunchModes::CustomRoles)
+		if (BuildCookRun->GetLaunchMode() == ELauncherProfileLaunchModes::CustomRoles)
 		{
-			ValidationErrors.Add(ELauncherProfileValidationErrors::CustomRolesNotSupportedYet);
+			ValidationData.Errors.Add(ELauncherProfileValidationErrors::CustomRolesNotSupportedYet);
 		}
 
 		// Launch: when using custom launch roles, all roles must have a device assigned
-		if (LaunchMode == ELauncherProfileLaunchModes::CustomRoles)
+		if (BuildCookRun->GetLaunchMode() == ELauncherProfileLaunchModes::CustomRoles)
 		{
-			for (int32 RoleIndex = 0; RoleIndex < LaunchRoles.Num(); ++RoleIndex)
+			for (int32 RoleIndex = 0; RoleIndex < BuildCookRun->GetLaunchRoles().Num(); ++RoleIndex)
 			{
-				if (LaunchRoles[RoleIndex]->GetAssignedDevice().IsEmpty())
+				if (BuildCookRun->GetLaunchRoles()[RoleIndex]->GetAssignedDevice().IsEmpty())
 				{
-					ValidationErrors.Add(ELauncherProfileValidationErrors::NoLaunchRoleDeviceAssigned);
+					ValidationData.Errors.Add(ELauncherProfileValidationErrors::NoLaunchRoleDeviceAssigned);
 				
 					break;
 				}
@@ -3422,126 +2570,135 @@ protected:
 		}
 
 		// Launch: Zen pak streaming requires deployment and launch
-		if (IsUsingZenPakStreaming() && (LaunchMode == ELauncherProfileLaunchModes::DoNotLaunch || DeploymentMode != ELauncherProfileDeploymentModes::CopyToDevice))
+		if (BuildCookRun->IsUsingZenPakStreaming() && (BuildCookRun->GetLaunchMode() == ELauncherProfileLaunchModes::DoNotLaunch || BuildCookRun->GetDeploymentMode() != ELauncherProfileDeploymentModes::CopyToDevice))
 		{
-			ValidationErrors.Add(ELauncherProfileValidationErrors::ZenPakStreamingRequiresDeployAndLaunch);
+			ValidationData.Errors.Add(ELauncherProfileValidationErrors::ZenPakStreamingRequiresDeployAndLaunch);
 		}
 
 		bool bLegacyIterative = false;
 		GConfig->GetBool(TEXT("CookSettings"), TEXT("LegacyIterative"), bLegacyIterative, GEditorIni);
 
-		if (CookUnversioned && bLegacyIterative && CookIncremental && ((CookMode == ELauncherProfileCookModes::ByTheBook) || (CookMode == ELauncherProfileCookModes::ByTheBookInEditor)))
+		if (BuildCookRun->IsCookingUnversioned() && bLegacyIterative && BuildCookRun->IsCookingIncrementally() && ((BuildCookRun->GetCookMode() == ELauncherProfileCookModes::ByTheBook) || (BuildCookRun->GetCookMode() == ELauncherProfileCookModes::ByTheBookInEditor)))
 		{
-			ValidationErrors.Add(ELauncherProfileValidationErrors::UnversionedAndIncremental);
+			ValidationData.Errors.Add(ELauncherProfileValidationErrors::UnversionedAndIncremental);
 		}
 
 
-		if ( (IsGeneratingPatch() || ShouldAddPatchLevel()) && (CookMode != ELauncherProfileCookModes::ByTheBook) )
+		if ( (BuildCookRun->IsGeneratingPatch() || BuildCookRun->ShouldAddPatchLevel()) && (BuildCookRun->GetCookMode() != ELauncherProfileCookModes::ByTheBook) )
 		{
-			ValidationErrors.Add(ELauncherProfileValidationErrors::GeneratingPatchesCanOnlyRunFromByTheBookCookMode);
+			ValidationData.Errors.Add(ELauncherProfileValidationErrors::GeneratingPatchesCanOnlyRunFromByTheBookCookMode);
 		}
 
-		if (ShouldAddPatchLevel() && !IsGeneratingPatch() )
+		if (BuildCookRun->ShouldAddPatchLevel() && !BuildCookRun->IsGeneratingPatch() )
 		{
-			ValidationErrors.Add(ELauncherProfileValidationErrors::GeneratingMultiLevelPatchesRequiresGeneratePatch);
+			ValidationData.Errors.Add(ELauncherProfileValidationErrors::GeneratingMultiLevelPatchesRequiresGeneratePatch);
 		}
 
-		if (ShouldStageBaseReleasePaks() && BasedOnReleaseVersionName.IsEmpty())
+		if (ShouldStageBaseReleasePaks() && BuildCookRun->GetBasedOnReleaseVersionName().IsEmpty()) // deprecated
 		{
-			ValidationErrors.Add(ELauncherProfileValidationErrors::StagingBaseReleasePaksWithoutABaseReleaseVersion);
+			ValidationData.Errors.Add(ELauncherProfileValidationErrors::StagingBaseReleasePaksWithoutABaseReleaseVersion);
 		}
 
-		if ( IsGeneratingChunks() && (CookMode != ELauncherProfileCookModes::ByTheBook) )
+		if ( BuildCookRun->IsGeneratingChunks() && (BuildCookRun->GetCookMode() != ELauncherProfileCookModes::ByTheBook) )
 		{
-			ValidationErrors.Add(ELauncherProfileValidationErrors::GeneratingChunksRequiresCookByTheBook);
+			ValidationData.Errors.Add(ELauncherProfileValidationErrors::GeneratingChunksRequiresCookByTheBook);
 		}
 
-		if (IsGeneratingChunks() && !IsPackingWithUnrealPak())
+		if (BuildCookRun->IsGeneratingChunks() && !BuildCookRun->IsPackingWithUnrealPak())
 		{
-			ValidationErrors.Add(ELauncherProfileValidationErrors::GeneratingChunksRequiresUnrealPak);
+			ValidationData.Errors.Add(ELauncherProfileValidationErrors::GeneratingChunksRequiresUnrealPak);
 		}
 
-		if (IsGenerateHttpChunkData() && !IsGeneratingChunks() && !IsCreatingDLC())
+		if (BuildCookRun->IsGenerateHttpChunkData() && !BuildCookRun->IsGeneratingChunks() && !BuildCookRun->IsCreatingDLC())
 		{
-			ValidationErrors.Add(ELauncherProfileValidationErrors::GeneratingHttpChunkDataRequiresGeneratingChunks);
+			ValidationData.Errors.Add(ELauncherProfileValidationErrors::GeneratingHttpChunkDataRequiresGeneratingChunks);
 		}
 
-		if (IsGenerateHttpChunkData() && (GetHttpChunkDataReleaseName().IsEmpty() || !FPaths::DirectoryExists(*GetHttpChunkDataDirectory())))
+		if (BuildCookRun->IsGenerateHttpChunkData() && (BuildCookRun->GetHttpChunkDataReleaseName().IsEmpty() || !FPaths::DirectoryExists(*BuildCookRun->GetHttpChunkDataDirectory())))
 		{
-			ValidationErrors.Add(ELauncherProfileValidationErrors::GeneratingHttpChunkDataRequiresValidDirectoryAndName);
+			ValidationData.Errors.Add(ELauncherProfileValidationErrors::GeneratingHttpChunkDataRequiresValidDirectoryAndName);
 		}
 
 		// Launch: when launching, all devices that the build is launched on must have content cooked for their platform
-		if (LaunchMode != ELauncherProfileLaunchModes::DoNotLaunch && CookMode != ELauncherProfileCookModes::OnTheFly && CookMode != ELauncherProfileCookModes::OnTheFlyInEditor)
+		if (BuildCookRun->GetLaunchMode() != ELauncherProfileLaunchModes::DoNotLaunch && BuildCookRun->GetCookMode() != ELauncherProfileCookModes::OnTheFly && BuildCookRun->GetCookMode() != ELauncherProfileCookModes::OnTheFlyInEditor)
 		{
 			// @todo ensure that launched devices have cooked content
 		}
 		
-		if ((CookMode == ELauncherProfileCookModes::OnTheFly) || (CookMode == ELauncherProfileCookModes::OnTheFlyInEditor))
+		if ((BuildCookRun->GetCookMode() == ELauncherProfileCookModes::OnTheFly) || (BuildCookRun->GetCookMode() == ELauncherProfileCookModes::OnTheFlyInEditor))
 		{
-			if (BuildConfiguration == EBuildConfiguration::Shipping)
+			if (BuildCookRun->GetBuildConfiguration() == EBuildConfiguration::Shipping)
 			{
 				// shipping doesn't support commandline options
-				ValidationErrors.Add(ELauncherProfileValidationErrors::ShippingDoesntSupportCommandlineOptionsCantUseCookOnTheFly);
+				ValidationData.Errors.Add(ELauncherProfileValidationErrors::ShippingDoesntSupportCommandlineOptionsCantUseCookOnTheFly);
 			}
 
 		}
 
-		if (CookMode == ELauncherProfileCookModes::OnTheFly)
+		if (BuildCookRun->GetCookMode() == ELauncherProfileCookModes::OnTheFly)
 		{
 
-			for (auto const& CookedPlatform : CookedPlatforms)
+			for (auto const& CookedPlatform : BuildCookRun->GetCookedPlatforms())
 			{
 				if (CookedPlatform.Contains("Server"))
 				{
-					ValidationErrors.Add(ELauncherProfileValidationErrors::CookOnTheFlyDoesntSupportServer);
+					ValidationData.Errors.Add(ELauncherProfileValidationErrors::CookOnTheFlyDoesntSupportServer);
 				}
 			}
 			
 
 		}
 
-		if (bArchive && ArchiveDir.IsEmpty())
+		if (BuildCookRun->IsArchiving() && BuildCookRun->GetArchiveDirectory().IsEmpty())
 		{
-			ValidationErrors.Add(ELauncherProfileValidationErrors::NoArchiveDirectorySpecified);
+			ValidationData.Errors.Add(ELauncherProfileValidationErrors::NoArchiveDirectorySpecified);
 		}
 
-		if (bUseIoStore && !DeployWithUnrealPak)
+		if (BuildCookRun->IsUsingIoStore() && !BuildCookRun->IsPackingWithUnrealPak())
 		{
-			ValidationErrors.Add(ELauncherProfileValidationErrors::IoStoreRequiresPakFiles);
+			ValidationData.Errors.Add(ELauncherProfileValidationErrors::IoStoreRequiresPakFiles);
 		}
 
-		ValidateBuildTarget();
-		ValidatePlatformSDKs();
-		ValidateDeviceStatus();
+		ValidateBuildTarget(BuildCookRun);
+		ValidatePlatformSDKs(BuildCookRun);
+		ValidateDeviceStatus(BuildCookRun);
 
-		RefreshCustomWarningsAndErrors();
+		RefreshCustomWarningsAndErrors(BuildCookRun.ToSharedPtr());
 	}
+
 	
-	void ValidatePlatformSDKs(void)
+	void ValidatePlatformSDKs(const ILauncherProfileBuildCookRunRef& BuildCookRun)
 	{
-		ValidationErrors.Remove(ELauncherProfileValidationErrors::NoPlatformSDKInstalled);
+		if (!UATCommandValidation.Contains(BuildCookRun.ToSharedPtr()))
+		{
+			return;
+		}
+
+
+		FValidationData& ValidationData = UATCommandValidation.FindChecked(BuildCookRun.ToSharedPtr());
+
+		ValidationData.Errors.Remove(ELauncherProfileValidationErrors::NoPlatformSDKInstalled);
 		
 		// Cook: ensure that all platform SDKs are installed
-		if (CookedPlatforms.Num() > 0)
+		if (BuildCookRun->GetCookedPlatforms().Num() > 0)
 		{
 			bool bProjectHasCode = false; // @todo: Does the project have any code?
 			FString NotInstalledDocLink;
-			for (auto PlatformName : CookedPlatforms)
+			for (auto PlatformName : BuildCookRun->GetCookedPlatforms())
 			{
 				const ITargetPlatform* Platform = GetTargetPlatformManager()->FindTargetPlatform(PlatformName);
 				if(!Platform || !Platform->IsSdkInstalled(bProjectHasCode, NotInstalledDocLink))
 				{
-					ValidationErrors.Add(ELauncherProfileValidationErrors::NoPlatformSDKInstalled);
+					ValidationData.Errors.Add(ELauncherProfileValidationErrors::NoPlatformSDKInstalled);
 					ILauncherServicesModule& LauncherServicesModule = FModuleManager::GetModuleChecked<ILauncherServicesModule>(TEXT("LauncherServices"));
 					LauncherServicesModule.BroadcastLauncherServicesSDKNotInstalled(PlatformName, NotInstalledDocLink);
 					if (!Platform)
 					{
-						CookedPlatforms.Remove(PlatformName);
+						BuildCookRun->RemoveCookedPlatform(PlatformName);
 					}
 					else
 					{
-						InvalidPlatform = PlatformName;
+						ValidationData.InvalidPlatform = PlatformName;
 					}
 					return;
 				}
@@ -3549,9 +2706,9 @@ protected:
 		}
 		
 		// Deploy: ensure that all the target device SDKs are installed
-		if ((DeploymentMode != ELauncherProfileDeploymentModes::DoNotDeploy) && DeployedDeviceGroup.IsValid())
+		if ((BuildCookRun->GetDeploymentMode() != ELauncherProfileDeploymentModes::DoNotDeploy) && BuildCookRun->GetDeployedDeviceGroup(false).IsValid())
 		{
-			const TArray<FString>& Devices = DeployedDeviceGroup->GetDeviceIDs();
+			const TArray<FString>& Devices = BuildCookRun->GetDeployedDeviceGroup(false)->GetDeviceIDs();
 			for(auto DeviceId : Devices)
 			{
 				ITargetDeviceServicesModule* TargetDeviceServicesModule = static_cast<ITargetDeviceServicesModule*>(FModuleManager::Get().LoadModule(TEXT("TargetDeviceServices")));
@@ -3568,10 +2725,10 @@ protected:
 						const ITargetPlatform* Platform = GetTargetPlatformManager()->FindTargetPlatform(PlatformName);
 						if(!Platform || !Platform->IsSdkInstalled(bProjectHasCode, NotInstalledDocLink))
 						{
-							ValidationErrors.Add(ELauncherProfileValidationErrors::NoPlatformSDKInstalled);
+							ValidationData.Errors.Add(ELauncherProfileValidationErrors::NoPlatformSDKInstalled);
 							ILauncherServicesModule& LauncherServicesModule = FModuleManager::GetModuleChecked<ILauncherServicesModule>(TEXT("LauncherServices"));
 							LauncherServicesModule.BroadcastLauncherServicesSDKNotInstalled(PlatformName, NotInstalledDocLink);
-							DeployedDeviceGroup->RemoveDevice(DeviceId);
+							BuildCookRun->GetDeployedDeviceGroup(false)->RemoveDevice(DeviceId);
 							return;
 						}
 					}
@@ -3580,13 +2737,15 @@ protected:
 		}
 	}
 
-	void ValidateDeviceStatus(void)
+	void ValidateDeviceStatus(const ILauncherProfileBuildCookRunRef& BuildCookRun)
 	{
-		ValidationErrors.Remove(ELauncherProfileValidationErrors::LaunchDeviceIsUnauthorized);
+		FValidationData& ValidationData = UATCommandValidation.FindChecked(BuildCookRun.ToSharedPtr());
 
-		if (DeployedDeviceGroup.IsValid())
+		ValidationData.Errors.Remove(ELauncherProfileValidationErrors::LaunchDeviceIsUnauthorized);
+
+		if (BuildCookRun->GetDeployedDeviceGroup(false).IsValid())
 		{
-			const TArray<FString>& Devices = DeployedDeviceGroup->GetDeviceIDs();
+			const TArray<FString>& Devices = BuildCookRun->GetDeployedDeviceGroup(false)->GetDeviceIDs();
 			for (auto DeviceId : Devices)
 			{
 				ITargetDeviceServicesModule* TargetDeviceServicesModule = static_cast<ITargetDeviceServicesModule*>(FModuleManager::Get().LoadModule(TEXT("TargetDeviceServices")));
@@ -3597,7 +2756,7 @@ protected:
 					{
 						if (!DeviceProxy->IsAuthorized())
 						{
-							ValidationErrors.Add(ELauncherProfileValidationErrors::LaunchDeviceIsUnauthorized);
+							ValidationData.Errors.Add(ELauncherProfileValidationErrors::LaunchDeviceIsUnauthorized);
 							return;
 						}
 					}
@@ -3606,15 +2765,17 @@ protected:
 		}
 	}
 
-	void ValidateBuildTarget()
+	void ValidateBuildTarget(const ILauncherProfileBuildCookRunRef& BuildCookRun)
 	{
+		FValidationData& ValidationData = UATCommandValidation.FindChecked(BuildCookRun.ToSharedPtr());
+
 		bool bBuildTargetIsRequired = false;
 		bool bBuildTargetIsSelected = false;
 
 		TArray<FString> CurrentBuildTargets;
-		if (BuildTargetSpecified)
+		if (BuildCookRun->HasBuildTargetSpecified())
 		{
-			CurrentBuildTargets = BuildTargets;
+			CurrentBuildTargets = BuildCookRun->GetBuildTargets();
 			if (CurrentBuildTargets.IsEmpty())
 			{
 				CurrentBuildTargets.Add(FString());
@@ -3626,7 +2787,7 @@ protected:
 		}
 
 
-		TSet<EBuildTargetType> CookTargetTypes = GetCookTargetTypes();
+		TSet<EBuildTargetType> CookTargetTypes = GetCookTargetTypes(BuildCookRun);
 		if (CookTargetTypes.Num() == 0)
 		{
 			CookTargetTypes.Add(EBuildTargetType::Game); // UAT defaults to 'Game' too
@@ -3668,20 +2829,20 @@ protected:
 		{
 			if (!BuildTarget.IsEmpty() && !bBuildTargetIsSelected)
 			{
-				ValidationErrors.Add(ELauncherProfileValidationErrors::BuildTargetCookVariantMismatch);
+				ValidationData.Errors.Add(ELauncherProfileValidationErrors::BuildTargetCookVariantMismatch);
 				break;
 			}
 		}
 
 		if (bBuildTargetIsRequired)
 		{
-			if (BuildTargetSpecified)
+			if (BuildCookRun->HasBuildTargetSpecified())
 			{
-				ValidationErrors.Add(ELauncherProfileValidationErrors::BuildTargetIsRequired);
+				ValidationData.Errors.Add(ELauncherProfileValidationErrors::BuildTargetIsRequired);
 			}
 			else
 			{
-				ValidationErrors.Add(ELauncherProfileValidationErrors::FallbackBuildTargetIsRequired);
+				ValidationData.Errors.Add(ELauncherProfileValidationErrors::FallbackBuildTargetIsRequired);
 			}
 		}
 	}
@@ -3694,8 +2855,8 @@ protected:
 		// collect the build targets for the current project, filtered to what we are currently wanting to cook. Do not show fallback project's build targets
 		if (HasProjectSpecified())
 		{
-			TSet<EBuildTargetType> CookTargetTypes = GetCookTargetTypes();
-			LatestExplicitBuildTargetNames = FGameProjectHelper::GetExplicitBuildTargetsForProject( GetProjectPath(), &CookTargetTypes );
+			//TSet<EBuildTargetType> CookTargetTypes = GetCookTargetTypes(); // fixme: is it critical to have a 'get all cook target types' ?
+			LatestExplicitBuildTargetNames = FGameProjectHelper::GetExplicitBuildTargetsForProject( GetProjectPath()/*, &CookTargetTypes*/ );
 		}
 
 		// notify listeners if the explicitly-required build targets have changed
@@ -3707,10 +2868,10 @@ protected:
 	}
 
 
-	TSet<EBuildTargetType> GetCookTargetTypes() const
+	TSet<EBuildTargetType> GetCookTargetTypes(const ILauncherProfileBuildCookRunRef& BuildCookRun) const
 	{
 		TSet<EBuildTargetType> CookTargetTypes;
-		for ( const FString& Variant : GetCookedPlatforms() )
+		for ( const FString& Variant : BuildCookRun->GetCookedPlatforms() )
 		{
 			if (Variant.EndsWith(TEXT("Client")))
 			{
@@ -3733,208 +2894,95 @@ protected:
 		return MoveTemp(CookTargetTypes);
 	}
 
+	void OnSelectedProjectChanged()
+	{
+		RefreshValidBuildTargets();
+	}
+
+
+private:
 
-	void OnLauncherDeviceGroupDeviceAdded(const ILauncherDeviceGroupRef& DeviceGroup, const FString& DeviceId)
+	FValidationData* GetValidationData(ILauncherProfileUATCommandPtr UATCommand)
 	{
-		if( DeviceGroup == DeployedDeviceGroup )
+		// temp. fix : use the merged validation errors so legacy PL can display them
+		if (!UATCommand.IsValid() && FirstBuildCookRun.IsValid())
 		{
-			ValidatePlatformSDKs();
+			UATCommand = FirstBuildCookRun;
 		}
-	}
-	
-	void OnLauncherDeviceGroupDeviceRemove(const ILauncherDeviceGroupRef& DeviceGroup, const FString& DeviceId)
-	{
-		if( DeviceGroup == DeployedDeviceGroup )
+
+		if (UATCommand.IsValid())
+		{
+			return UATCommandValidation.Find(UATCommand);
+		}
+		else
 		{
-			ValidatePlatformSDKs();
+			return &ProfileValidation;
 		}
 	}
 
-	void OnSelectedProjectChanged()
+	const FValidationData* GetValidationData(ILauncherProfileUATCommandPtr UATCommand) const
 	{
-		RefreshValidBuildTargets();
+		// temp. fix : use the merged validation errors so legacy PL can display them
+		if (!UATCommand.IsValid() && FirstBuildCookRun.IsValid())
+		{
+			UATCommand = FirstBuildCookRun;
+		}
 
-		BuildTargets.Reset();
-		BuildTargetSpecified = ProjectSpecified; // if this is for 'Any Project' then it should use the fallback build target
+		if (UATCommand.IsValid())
+		{
+			return UATCommandValidation.Find(UATCommand);
+		}
+		else
+		{
+			return &ProfileValidation;
+		}
 	}
 
-private:
-	void QueryZenSnapshotsForProject();
-	TSharedPtr<UE::Zen::Build::FBuildListRetriever> GetBuildListRetriever();
 
-	//  Holds a reference to the launcher profile manager.
-	ILauncherProfileManagerRef LauncherProfileManager;
+	FString GetUniqueUATCommandName( const TCHAR* Base ) const
+	{
+		int32 UniqueIndex = 2;
+		FString UniqueName = Base;
+		while (GetUATCommand(*UniqueName) != nullptr)
+		{
+			UniqueName = FString::Printf(TEXT("%s-%d"), Base, UniqueIndex++);
+		}
 
-	// Holds the desired build configuration (only used if creating new builds).
-	EBuildConfiguration BuildConfiguration;
+		return UniqueName;
+	}
 
-	// Holds a flag indicating whether the build target is specified by this profile.
-	bool BuildTargetSpecified;
 
-	// Holds the names of the build targets (matching a .cs file) to build. Needed when multiple targets of a type exist
-	TArray<FString> BuildTargets;
+	ILauncherProfileUATCommandRef CreateUATCommandInternal( const TCHAR* InternalName = nullptr );
+	ILauncherProfileAutomatedTestRef CreateAutomatedTestInternal( const TCHAR* InternalName = nullptr );
+	ILauncherProfileBuildCookRunRef CreateBuildCookRunInternal( const TCHAR* InternalName = nullptr );
+	ILauncherProfileUATCommandPtr CreateUATCommandByTypeNameInternal( const TCHAR* TypeName, const TCHAR* InternalName = nullptr );
+	ILauncherProfileBuildCookRunPtr LoadBuildCookRunInternal(const FJsonObject& Object, int32 Version);
 
-	// Holds the build mode.
-	// Holds the build configuration name of the cooker.
-	EBuildConfiguration CookConfiguration;
+	void AddUATCommandInternal( const ILauncherProfileUATCommandRef& UATCommand )
+	{
+		UATCommands.Add(UATCommand);
+		UATCommandValidation.Add(UATCommand.ToSharedPtr());
+	}
 
-	// Holds additional cooker command line options.
-	FString CookOptions;
 
-	// Holds the cooking mode.
-	TEnumAsByte<ELauncherProfileCookModes::Type> CookMode;
+	void QueryZenSnapshotsForProject();
+	TSharedPtr<UE::Zen::Build::FBuildListRetriever> GetBuildListRetriever();
 
-	// Holds the build mode
-	TEnumAsByte<ELauncherProfileBuildModes::Type> BuildMode;
+	//  Holds a reference to the launcher profile manager.
+	ILauncherProfileManagerRef LauncherProfileManager;
 
-	// Holds a flag indicating whether fast iteration is enabled 
-	bool bFastIterate;
 	
 	// Holds a flag indicating whether UAT should be built
 	bool BuildUAT;
-
-	// Generate compressed content
-	bool Compressed;
-
-	// encrypt ini files in the pak file
-	bool EncryptIniFiles;
-	// is this build for distribution
-	bool ForDistribution;
-	
-	// Holds a flag indicating whether only modified content should be cooked.
-	bool CookIncremental;
-
-	// hold a flag indicating if we want to iterate on the shared cooked build
-	bool IterateSharedCookedBuild;
-
-	// Holds a flag indicating whether packages should be saved without a version.
-	bool CookUnversioned;
-
-	bool bSkipCookingEditorContent;
-
-	// This setting is used only if cooking by the book (only used if cooking by the book).
-	TArray<FString> CookedCultures;
-
-	// Holds the collection of maps to cook (only used if cooking by the book).
-	TArray<FString> CookedMaps;
-
-	// Holds the platforms to include in the build (only used if creating new builds).
-	TArray<FString> CookedPlatforms;
-
-	// Holds the platforms to deploy to if no specific devices were chosen for deploy.
-	FName DefaultDeployPlatform;
-
-	// Holds the default role (only used if launch mode is DefaultRole).
-	ILauncherProfileLaunchRoleRef DefaultLaunchRole;
-
-	// Holds a flag indicating whether a streaming file server should be used.
-	bool DeployStreamingServer;
-
-	// Holds a flag indicating whether content should be packaged with UnrealPak.
-	bool DeployWithUnrealPak;
-
-	// Flag to indicate if game prerequisites should be included
-	bool IncludePrerequisites;
-
-	// Flag indicating if content should be split into chunks
-	bool bGenerateChunks;
-	
-	// Flag indicating if chunked content should be used to generate HTTPChunkInstall data
-	bool bGenerateHttpChunkData;
-	
-	// Where to store HTTPChunkInstall data
-	FString HttpChunkDataDirectory;
-	
-	// Version name of the HTTPChunkInstall data
-	FString HttpChunkDataReleaseName;
-
-	// if present, iostore container creation will try to use existing compressed blocks
-	// instead of compressing new ones, to avoid patches from compressor version changes,
-	// and to speed up iostore container creation time. See IoStoreUtilities.cpp ReferenceContainerGlobalFileName.
-	FString ReferenceContainerGlobalFileName;
-
-	// If ReferenceContainerGlobalFileName refers to encrypted containers, this is the filename of
-	// the json file containing the keys.
-	FString ReferenceContainerCryptoKeysFileName;
-
-	// create a release version of the content (this can be used to base dlc / patches from)
-	bool CreateReleaseVersion;
-
-	// name of the release version
-	FString CreateReleaseVersionName;
-
-	// name of the release version to base this dlc / patch on
-	FString BasedOnReleaseVersionName;
-
-	// name of the original release version
-	FString OriginalReleaseVersionName;
-
-	// This build generate a patch based on some source content seealso PatchSourceContentPath
-	bool GeneratePatch;
-
-	// This build generates a new tier patch file for modified content
-	bool AddPatchLevel;
-
 	// This build stages pak files from the release version it is based on (deprecated)
 	bool StageBaseReleasePaks;
-
-	// This build will cook content for dlc See also DLCName
-	bool CreateDLC;
-
-	// name of the dlc we are going to build (the name of the dlc plugin)
-	FString DLCName;
-
-	// should the dlc include engine content in the current dlc 
-	//  engine content which was not referenced by original release
-	//  otherwise error on any access of engine content during dlc cook
-	bool DLCIncludeEngineContent;
-
-	// Holds a flag indicating whether to use incremental deployment
-	bool DeployIncremental;
-
-	// Holds the device group to deploy to.
-	ILauncherDeviceGroupPtr DeployedDeviceGroup;
-
-	// Delegate handles for registered DeployedDeviceGroup event handlers.
-	FDelegateHandle OnLauncherDeviceGroupDeviceAddedDelegateHandle;
-	FDelegateHandle OnLauncherDeviceGroupDeviceRemoveDelegateHandle;
-
-	// Holds the identifier of the deployed device group.
-	FGuid DeployedDeviceGroupId;
-
-	// Holds the deployment mode.
-	TEnumAsByte<ELauncherProfileDeploymentModes::Type> DeploymentMode;
-
-	// Holds a flag indicating whether the file server's console window should be hidden.
-	bool HideFileServerWindow;
-
 	// Holds the profile's unique identifier.
 	FGuid Id;
-
-	// Holds the launch mode.
-	TEnumAsByte<ELauncherProfileLaunchModes::Type> LaunchMode;
-
-	// Holds the launch roles (only used if launch mode is UsingRoles).
-	TArray<ILauncherProfileLaunchRolePtr> LaunchRoles;
-
 	// Holds the profile's name.
 	FString Name;
 
 	// Holds the profile's description.
 	FString Description;
-
-	// Holds the packaging mode.
-	TEnumAsByte<ELauncherProfilePackagingModes::Type> PackagingMode;
-
-	// Holds the package directory
-	FString PackageDir;
-
-	// Whether to run archive step	
-	bool bArchive;
-
-	// Holds the archive directory
-	FString ArchiveDir;
-
 	// Holds a flag indicating whether the project is specified by this profile.
 	bool ProjectSpecified;
 
@@ -3945,22 +2993,9 @@ private:
 	// otherwise this is an absolute path.
 	FString ShareableProjectPath;
 
-	// Holds the full path to a folder to use for Zen pak streaming
-	FString ZenPakStreamingPath;
-
-	// Holds the collection of validation errors.
-	TArray<ELauncherProfileValidationErrors::Type> ValidationErrors;
-	FString InvalidPlatform;
-
-	// Holds the collection of custom errors & warnings.
-	TMap<FString,FText> CustomErrors;
-	TMap<FString,FText> CustomWarnings;
-
-    // Holds the time out time for the cook on the fly server
-    uint32 Timeout;
-
-    // Holds the close value for the cook on the fly server
-    bool ForceClose;
+	// Profile validation
+	FValidationData ProfileValidation;
+	TMap<ILauncherProfileUATCommandPtr, FValidationData> UATCommandValidation;
 
 	// Path to the editor executable to pass to UAT, for cooking, etc... May be empty.
 	FString EditorExe;
@@ -3968,38 +3003,9 @@ private:
 	// Profile is for an internal project
 	bool bNotForLicensees;
 
-	// Additional command line parameters to set for the application when it launches
-	FString AdditionalCommandLineParameters;
-
-	// Additional command line parameters to set for a given build target type when it launches
-	TMap<EBuildTargetType,FString> AdditionalTargetCommandLineParameters;
-
-	// Use I/O store.
-	bool bUseIoStore;
-	
-	// Use Zen storage server
-	bool bUseZenStore;
-
-	// Import a Zen snapshot before cooking
-	bool bImportingZenSnapshot;
-
-	// Use Zen pak streaming
-	bool bUseZenPakStreaming;
-
-	// Make binary config.
-	bool bMakeBinaryConfig;
-
 	// Update flash on device before running
 	bool bShouldUpdateFlash;
 
-	// Is the launch device actually a Simulator
-	bool bIsDeviceASimulator;
-
-	// Architectures to build (x64, arm64 etc)
-	TArray<FString> ClientArchitectures;
-	TArray<FString> ServerArchitectures;
-	TArray<FString> EditorArchitectures;
-
 	// Key-Value pairs for extsibility. Not user-facing.
 	TMap<FString,FString> CustomStringProperties;
 	TMap<FString, bool>   CustomBoolProperties;
@@ -4010,9 +3016,15 @@ private:
 
 	// Extra UAT commands
 	TArray<ILauncherProfileUATCommandRef> UATCommands;
+	FOnLauncherProfileUATCommandChanged UATCommandAdded;
+	FOnLauncherProfileUATCommandChanged UATCommandRemoved;
 
 private:
 
+	// Cache of the first BuildCookRun command in the list, for convenience & for backwards compatibilty with Legacy Project Launcher
+	// NOTE: a profile is not guaranteed to have a BuildCookRun command - it can be deleted
+	TSharedPtr<ILauncherProfileBuildCookRun> FirstBuildCookRun;
+
 	// cook in the editor callbacks (not valid for any other cook mode)
 	FIsCookFinishedDelegate IsCookFinishedDelegate;
 	FCookCanceledDelegate CookCanceledDelegate;
@@ -4028,6 +3040,7 @@ private:
 
 	// Holds a delegate to be invoked when the profile is validated to gather custom warnings and errors
 	FOnProfileCustomValidation CustomValidationDelegate;
+	FOnProfileCustomUATCommandValidation CustomUATCommandValidationDelegate;
 
 	// Cached build target options (not serialized)
 	TArray<FString> ExplictBuildTargetNames;
@@ -4038,4 +3051,6 @@ private:
 	// Cache all builds for projects and platforms, shared between profiles
 	typedef TMap<FString, int32> FPlatformToBuildsMap;
 	static TMap<FString, FPlatformToBuildsMap> PerProjectBuilds;
+
+	friend class FLauncherProfileBuildCookRun; // for legacy serialize
 };
