Launcher Profile support for multiple build cook run commands per-profile.

diff --git a/Engine/Source/Developer/LauncherServices/Private/Profiles/LauncherProfileAutomatedTest.h b/Engine/Source/Developer/LauncherServices/Private/Profiles/LauncherProfileAutomatedTest.h
--- a/Engine/Source/Developer/LauncherServices/Private/Profiles/LauncherProfileAutomatedTest.h
+++ b/Engine/Source/Developer/LauncherServices/Private/Profiles/LauncherProfileAutomatedTest.h
@@ -4,23 +4,20 @@
 
 #include "CoreMinimal.h"
 #include "ILauncherProfileAutomatedTest.h"
+#include "ILauncherProfileBuildCookRun.h"
 #include "Dom/JsonObject.h"
+#include "Internationalization/Internationalization.h"
 
 class FLauncherProfileAutomatedTest final : public ILauncherProfileAutomatedTest, public TSharedFromThis<FLauncherProfileAutomatedTest>
 {
 public:
 	static const TCHAR* TypeName;
 
-	FLauncherProfileAutomatedTest( const TCHAR* InInternalName ) 
+	FLauncherProfileAutomatedTest( const TCHAR* InInternalName = TEXT("") ) 
 		: InternalName(InInternalName)
 	{
 	}
 
-	FLauncherProfileAutomatedTest(const FJsonObject& Object)
-	{
-		Load(Object);
-	}
-
 	//UE_DEPRECATED(5.8, "please use the Json Save/Load API") - confuses linux build which thinks the entire class is deprecated
 	FLauncherProfileAutomatedTest( FArchive& Archive )
 	{
@@ -48,13 +45,13 @@ public:
 		return Tests;
 	}
 
-	virtual void SetAdditionalCommandLine( const TCHAR* InAdditionalCommandLine ) override
+	virtual void SetAdditionalUATCommandLine( const TCHAR* InAdditionalCommandLine ) override
 	{
-		AdditionalCommandLine = InAdditionalCommandLine;
+		UATCommandLine = InAdditionalCommandLine;
 	}
-	virtual const FString& GetAdditionalCommandLine() const override
+	virtual const FString& GetAdditionalUATCommandLine() const override
 	{
-		return AdditionalCommandLine;
+		return UATCommandLine;
 	}
 
 	virtual void SetEnabled( bool bInEnabled ) override
@@ -80,32 +77,82 @@ public:
 		return *InternalName;
 	}
 
+	virtual void SetDescription(const FString& NewDescription) override
+	{
+		Description = NewDescription;
+	}
+
+	virtual FString GetDescription() const
+	{
+		return Description;
+	}
+
 	virtual ILauncherProfileAutomatedTestPtr AsAutomatedTest() override
 	{
 		return AsShared().ToSharedPtr();
 	}
 
+	virtual ILauncherProfileBuildCookRunPtr AsBuildCookRun() override
+	{
+		return nullptr;
+	}
 
-	virtual void Load(const FJsonObject& Object) override
+	enum class EVersion : uint32
 	{
+		Original = 0,
+		FirstVersioned = 1,
+		// add new items here
+
+		Final,
+		Latest = Final-1,
+	};
+
+
+	virtual bool Load(const FJsonObject& Object) override
+	{
+		EVersion Version = EVersion::Original;
+		if (!Object.TryGetNumberField(TEXT("Version"), (uint32&)Version))
+		{
+			Version = EVersion::Original;
+		}
+		if (Version > EVersion::Latest)
+		{
+			return false;
+		}
+
 		InternalName          = Object.GetStringField(TEXT("InternalName"));
 		UATCommand            = Object.GetStringField(TEXT("UATCommand"));
 		Tests                 = Object.GetStringField(TEXT("Tests"));
-		AdditionalCommandLine = Object.GetStringField(TEXT("AdditionalCommandLine"));
+		UATCommandLine        = Object.GetStringField(TEXT("AdditionalCommandLine"));
 		bEnabled              = Object.GetBoolField(TEXT("Enabled"));
-		Order                 = Object.GetIntegerField(TEXT("Order"));
+
+
+		if (Version >= EVersion::FirstVersioned)
+		{
+			Description = Object.GetStringField(TEXT("Description"));
+			Order       = Object.GetIntegerField(TEXT("Order"));
+
+		}
+		else
+		{
+			Object.TryGetNumberField(TEXT("Order"), Order);
+		}
 		Order = FMath::Max(1, Order);
+		return true;
 	}
+
 	virtual void Save(TJsonWriter<>& Writer) override
 	{
 		Writer.WriteObjectStart();
 		Writer.WriteValue(TEXT("Type"),                  TypeName);
+		Writer.WriteValue(TEXT("Version"),               (uint32)EVersion::Latest);
 		Writer.WriteValue(TEXT("InternalName"),          InternalName);
 		Writer.WriteValue(TEXT("UATCommand"),            UATCommand);
 		Writer.WriteValue(TEXT("Tests"),                 Tests);
-		Writer.WriteValue(TEXT("AdditionalCommandLine"), AdditionalCommandLine);
+		Writer.WriteValue(TEXT("AdditionalCommandLine"), UATCommandLine);
 		Writer.WriteValue(TEXT("Enabled"),               bEnabled);
 		Writer.WriteValue(TEXT("Order"),                 Order);
+		Writer.WriteValue(TEXT("Description"),           Description);
 		Writer.WriteObjectEnd();
 	}
 
@@ -115,7 +162,7 @@ public:
 		Archive << InternalName
 				<< UATCommand
 				<< Tests
-				<< AdditionalCommandLine
+				<< UATCommandLine
 				<< bEnabled
 				<< Order
 		;
@@ -123,14 +170,16 @@ public:
 
 
 protected:
+
 	FString UATCommand = TEXT("RunUnreal");
 	FString Tests; //-Test=xxx,yyy
-	FString AdditionalCommandLine;
+	FString UATCommandLine;
 
 	bool bEnabled = true;
 	int32 Order = 1; // longer tests should generally have a higher order
 
 	FString InternalName;
+	FString Description = NSLOCTEXT("FLauncherProfileAutomatedTest", "DefaultDescripton", "Automated Test").ToString();
 };
 
 
diff --git a/Engine/Source/Developer/LauncherServices/Private/Profiles/LauncherProfileManager.cpp b/Engine/Source/Developer/LauncherServices/Private/Profiles/LauncherProfileManager.cpp
--- a/Engine/Source/Developer/LauncherServices/Private/Profiles/LauncherProfileManager.cpp
+++ b/Engine/Source/Developer/LauncherServices/Private/Profiles/LauncherProfileManager.cpp
@@ -184,9 +184,12 @@ void FLauncherProfileManager::AddProfile( const ILauncherProfileRef& Profile )
 			RemoveProfile(ExistingProfile.ToSharedRef());
 		}
 
-		if (!Profile->GetDeployedDeviceGroup().IsValid())
+		for (const ILauncherProfileBuildCookRunRef& BuildCookRun : Profile->GetBuildCookRunCommands())
 		{
-			Profile->SetDeployedDeviceGroup(AddNewDeviceGroup());
+			if (!BuildCookRun->GetDeployedDeviceGroup().IsValid())
+			{
+				BuildCookRun->SetDeployedDeviceGroup(AddNewDeviceGroup());
+			}
 		}
 
 		// add the new profile
@@ -266,12 +269,15 @@ PRAGMA_DISABLE_DEPRECATION_WARNINGS
 	if (Profile->Serialize(Archive))
 PRAGMA_ENABLE_DEPRECATION_WARNINGS
 	{
-		ILauncherDeviceGroupPtr DeviceGroup = GetDeviceGroup(Profile->GetDeployedDeviceGroupId());
-		if (!DeviceGroup.IsValid())
+		for (const ILauncherProfileBuildCookRunRef& BuildCookRun : Profile->GetBuildCookRunCommands())
 		{
-			DeviceGroup = AddNewDeviceGroup();	
+			ILauncherDeviceGroupPtr DeviceGroup = BuildCookRun->GetDeployedDeviceGroup();
+			if (!DeviceGroup.IsValid())
+			{
+				DeviceGroup = AddNewDeviceGroup();	
+			}
+			BuildCookRun->SetDeployedDeviceGroup(DeviceGroup);
 		}
-		Profile->SetDeployedDeviceGroup(DeviceGroup);
 
 		return MakeShareable(Profile);
 	}
@@ -298,13 +304,6 @@ ILauncherProfilePtr FLauncherProfileManager::LoadJSONProfile(FString ProfileFile
 
 	if (Profile->Load(*(Object.Get())))
 	{
-		ILauncherDeviceGroupPtr DeviceGroup = GetDeviceGroup(Profile->GetDeployedDeviceGroupId());
-		if (!DeviceGroup.IsValid())
-		{
-			DeviceGroup = AddNewDeviceGroup();
-		}
-		Profile->SetDeployedDeviceGroup(DeviceGroup);
-
 		return MakeShareable(Profile);
 	}
 
diff --git a/Engine/Source/Developer/LauncherServices/Private/Profiles/LauncherProfileUATCommand.h b/Engine/Source/Developer/LauncherServices/Private/Profiles/LauncherProfileUATCommand.h
--- a/Engine/Source/Developer/LauncherServices/Private/Profiles/LauncherProfileUATCommand.h
+++ b/Engine/Source/Developer/LauncherServices/Private/Profiles/LauncherProfileUATCommand.h
@@ -5,23 +5,20 @@
 #include "CoreMinimal.h"
 #include "ILauncherProfileUATCommand.h"
 #include "ILauncherProfileAutomatedTest.h"
+#include "ILauncherProfileBuildCookRun.h"
 #include "Dom/JsonObject.h"
+#include "Internationalization/Internationalization.h"
 
 class FLauncherProfileUATCommand final : public ILauncherProfileUATCommand, public TSharedFromThis<FLauncherProfileUATCommand>
 {
 public:
 	static const TCHAR* TypeName;
 
-	FLauncherProfileUATCommand( const TCHAR* InInternalName ) 
+	FLauncherProfileUATCommand(const TCHAR* InInternalName = TEXT("") ) 
 		: InternalName(InInternalName)
 	{
 	}
 
-	FLauncherProfileUATCommand(const FJsonObject& Object)
-	{
-		Load(Object);
-	}
-
 	virtual void SetUATCommand( const TCHAR* InUATCommand ) override
 	{
 		UATCommand = InUATCommand;
@@ -32,13 +29,13 @@ public:
 		return UATCommand;
 	}
 
-	virtual void SetAdditionalCommandLine( const TCHAR* InAdditionalCommandLine ) override
+	virtual void SetAdditionalUATCommandLine( const TCHAR* InAdditionalCommandLine ) override
 	{
-		AdditionalCommandLine = InAdditionalCommandLine;
+		UATCommandLine = InAdditionalCommandLine;
 	}
-	virtual const FString& GetAdditionalCommandLine() const override
+	virtual const FString& GetAdditionalUATCommandLine() const override
 	{
-		return AdditionalCommandLine;
+		return UATCommandLine;
 	}
 
 	virtual void SetEnabled( bool bInEnabled ) override
@@ -52,7 +49,6 @@ public:
 
 	virtual void SetOrder( int32 InOrder ) override
 	{
-		check(InOrder != 0);
 		Order = InOrder;
 	}
 	virtual int32 GetOrder() const override
@@ -65,41 +61,85 @@ public:
 		return *InternalName;
 	}
 
+	virtual void SetDescription(const FString& NewDescription) override
+	{
+		Description = NewDescription;
+	}
+
+	virtual FString GetDescription() const
+	{
+		return Description;
+	}
+
 	virtual ILauncherProfileAutomatedTestPtr AsAutomatedTest() override
 	{
 		return nullptr;
 	}
 
+	virtual ILauncherProfileBuildCookRunPtr AsBuildCookRun() override
+	{
+		return nullptr;
+	}
 
-	virtual void Load(const FJsonObject& Object) override
+	enum class EVersion : uint32
 	{
+		Original = 0,
+		FirstVersioned = 1,
+		// add new items here
+
+		Final,
+		Latest = Final-1,
+	};
+
+	virtual bool Load(const FJsonObject& Object) override
+	{
+		EVersion Version = EVersion::Original;
+		if (!Object.TryGetNumberField(TEXT("Version"), (uint32&)Version))
+		{
+			Version = EVersion::Original;
+		}
+		if (Version > EVersion::Latest)
+		{
+			return false;
+		}
+
 		InternalName          = Object.GetStringField(TEXT("InternalName"));
 		UATCommand            = Object.GetStringField(TEXT("UATCommand"));
-		AdditionalCommandLine = Object.GetStringField(TEXT("AdditionalCommandLine"));
+		UATCommandLine        = Object.GetStringField(TEXT("AdditionalCommandLine"));
 		bEnabled              = Object.GetBoolField(TEXT("Enabled"));
 		Order                 = Object.GetIntegerField(TEXT("Order"));
+		if (Version >= EVersion::FirstVersioned)
+		{
+			Description = Object.GetStringField(TEXT("Description"));
+		}
+		return true;
 	}
+
 	virtual void Save(TJsonWriter<>& Writer) override
 	{
 		Writer.WriteObjectStart();
 		Writer.WriteValue(TEXT("Type"),                  TypeName);
+		Writer.WriteValue(TEXT("Version"),               (uint32)EVersion::Latest);
 		Writer.WriteValue(TEXT("InternalName"),          InternalName);
 		Writer.WriteValue(TEXT("UATCommand"),            UATCommand);
-		Writer.WriteValue(TEXT("AdditionalCommandLine"), AdditionalCommandLine);
+		Writer.WriteValue(TEXT("AdditionalCommandLine"), UATCommandLine);
 		Writer.WriteValue(TEXT("Enabled"),               bEnabled);
 		Writer.WriteValue(TEXT("Order"),                 Order);
+		Writer.WriteValue(TEXT("Description"),           Description);
 		Writer.WriteObjectEnd();
 	}
 
 
 protected:
+
 	FString UATCommand = TEXT("RunUnreal");
-	FString AdditionalCommandLine;
+	FString UATCommandLine;
 
 	bool bEnabled = true;
 	int32 Order = -100000; // negative order items will run before the main BuildCookRun, positive order items will run afterwards
 
 	FString InternalName;
+	FString Description = NSLOCTEXT("FLauncherProfileUATCommand", "DefaultDescripton", "Additional Task").ToString();
 };
 
 
diff --git a/Engine/Source/Developer/LauncherServices/Public/ILauncherProfile.h b/Engine/Source/Developer/LauncherServices/Public/ILauncherProfile.h
--- a/Engine/Source/Developer/LauncherServices/Public/ILauncherProfile.h
+++ b/Engine/Source/Developer/LauncherServices/Public/ILauncherProfile.h
@@ -8,268 +8,15 @@
 #include "ILauncherDeviceGroup.h"
 #include "ILauncherProfileLaunchRole.h"
 #include "ILauncherProfileAutomatedTest.h"
+#include "ILauncherProfileBuildCookRun.h"
+
+DECLARE_LOG_CATEGORY_EXTERN(LogLauncherProfile, Log, All);
 
 class Error;
 class FJsonObject;
 class ILauncherProfile;
 class ILauncherSimpleProfile;
 
-namespace ELauncherProfileBuildModes
-{
-	/**
-	 * Enumerates modes in which the launcher cooks builds.
-	 */
-	enum Type
-	{
-		/** Build if there is not already an existing pre-built target available. */
-		Auto,
-
-		/** Always build. */
-		Build,
-
-		/** Do not build. */
-		DoNotBuild,
-	};
-}
-
-
-
-namespace ELauncherProfileCookModes
-{
-	/**
-	 * Enumerates modes in which the launcher cooks builds.
-	 */
-	enum Type
-	{
-		/** Do not cook the build (default). */
-		DoNotCook,
-
-		/** Pre-cook using user specified settings. */
-		ByTheBook,
-
-		/** Cook the build on the fly while the game is running. */
-		OnTheFly,
-
-		/** Cook by the book in the editor process space */
-		ByTheBookInEditor,
-
-		/** Cook on the fly in the editor process space */
-		OnTheFlyInEditor,
-	};
-}
-
-
-namespace ELauncherProfileCookedMaps
-{
-	/**
-	 * Enumerates selections for maps to cook.
-	 */
-	enum Type
-	{
-		/** Cook all maps. */
-		AllMaps,
-
-		/** Don't cook any maps. Only startup packages will be cooked. */
-		NoMaps,
-
-		/** Cook user selected maps. */
-		SelectedMaps
-
-	};
-}
-
-
-namespace ELauncherProfileDeploymentModes
-{
-	/**
-	 * Enumerates deployment modes.
-	 */
-	enum Type
-	{
-		/** Do not deploy the build to any device. */
-		DoNotDeploy,
-
-		/** Copy all required file to the device. */
-		CopyToDevice,
-
-		/** Let the device get required files from a file server. */
-		FileServer,
-
-		/** Copy a build from a repository to the device. */
-		 CopyRepository,
-	};
-}
-
-
-namespace ELauncherProfileLaunchModes
-{
-	/**
-	 * Enumerates launch modes.
-	 */
-	enum Type
-	{
-		/** Do not launch. */
-		DoNotLaunch,
-
-		/** Launch with customized roles per device. */
-		CustomRoles,
-
-		/** Launch with the default role on all deployed devices. */
-		DefaultRole	
-	};
-}
-
-
-namespace ELauncherProfilePackagingModes
-{
-	/**
-	 * Enumerates packaging modes.
-	 */
-	enum Type
-	{
-		/** Do not package. */
-		DoNotPackage,
-
-		/** Package and store the build locally. */
-		Locally,
-
-		/** Package and store the build in a shared repository. */
-		SharedRepository
-	};
-}
-
-
-namespace ELauncherProfileIncrementalCookMode
-{
-	/**
-	 * Enumerates modes in which the launcher handles incremental cooking
-	 */
-	enum Type
-	{
-		/** Do not cook incrementally (default). */
-		None,
-
-		/** -iterativecooking : Only recook specifically modified assets without recooking dependencies. This is the legacy method and is must less robust. */
-		ModifiedOnly,
-
-		/** -cookincremental : Recook modified assets and any dependencies. This is the new method and is currently slower but far more robust. */
-		ModifiedAndDependencies,
-	};
-}
-
-
-
-namespace ELauncherProfileValidationErrors
-{
-	/**
-	 * Enumerates profile validation messages.
-	 */
-	enum Type
-	{
-		/**
-		 * Deployment by copying required files to a device requires
-		 * cooking by the book and is incompatible with cook on the fly.
-		 */
-		CopyToDeviceRequiresCookByTheBook,
-
-		/** Custom launch roles are not yet supported. */
-		CustomRolesNotSupportedYet,
-
-		/** A device group must be selected when deploying builds. */
-		DeployedDeviceGroupRequired,
-		
-		/** The initial culture configured for launch is not part of the selected build. */
-		InitialCultureNotAvailable,
-
-		/** The initial map configured for launch is not part of the selected build. */
-		InitialMapNotAvailable,
-
-		/** The specified launch command line is not formatted correctly. */
-		MalformedLaunchCommandLine,
-
-		/** A build configuration is required when creating new builds. */
-		NoBuildConfigurationSelected,
-
-		/** When cooking a build, at least one culture must be included. */
-		NoCookedCulturesSelected,
-
-		/** One or more launch roles do not have a device assigned. */
-		NoLaunchRoleDeviceAssigned,
-
-		/** At least one platform is required when creating new builds. */
-		NoPlatformSelected,
-
-		/** A game is required when creating new builds. */
-		NoProjectSelected,
-
-		/** The deployment requires a package directory to be specified */
-		NoPackageDirectorySpecified,
-		
-		/** The platform SDK is not installed but is required. */
-		NoPlatformSDKInstalled,
-
-		/** The profile has unversioned and incremental specified these are not compatible together */
-		UnversionedAndIncremental,
-
-		/** Deprecated, do not use */
-		UnversionedAndIncrimental UE_DEPRECATED(5.6, "Update to use the correctly spelt 'ELauncherProfileValidationErrors::UnversionedAndIncremental' instead") = UnversionedAndIncremental,
-
-		/** generating patch requires cook by the book mode*/
-		GeneratingPatchesCanOnlyRunFromByTheBookCookMode,
-
-		/** generating multilevel patch requires generating patch */
-		GeneratingMultiLevelPatchesRequiresGeneratePatch,
-
-		/** staging base release pak files requires a base release version to be specified */
-		StagingBaseReleasePaksWithoutABaseReleaseVersion,
-
-		/** Generating Chunks requires cook by the book mode */
-		GeneratingChunksRequiresCookByTheBook,
-
-		/** Generating Chunks requires UnrealPak */
-		GeneratingChunksRequiresUnrealPak,
-
-		/** Generating http chunk install data requires generating chunks or DLC*/
-		GeneratingHttpChunkDataRequiresGeneratingChunks,
-
-		/** Generating http chunk install data requires valid install directorys and release name */
-		GeneratingHttpChunkDataRequiresValidDirectoryAndName,
-
-		/** Shipping doesn't support commandline options can't use cook on the fly */
-		ShippingDoesntSupportCommandlineOptionsCantUseCookOnTheFly,
-
-		/** Cook on the fly doesn't support server target platforms */
-		CookOnTheFlyDoesntSupportServer,
-
-		/** The archive step requires a directory to be specified */
-		NoArchiveDirectorySpecified,
-
-		/** Device is unauthorized or is locked */
-		LaunchDeviceIsUnauthorized,
- 
-		/** Using I/O store container file(s) requires using UnrealPak */
-		IoStoreRequiresPakFiles,
-
-		/** Build Target and Cook Variant mismatch */
-		BuildTargetCookVariantMismatch,
-
-		/** Build Target is required */
-		BuildTargetIsRequired,
-
-		/** Fallback Build Target is required */
-		FallbackBuildTargetIsRequired,
-
-		/** Packaging and deploying are mutually exclusive */
-		CopyToDeviceRequiresNoPackaging,
-
-		/** Zen Pak Streaming requires deployment and launching */
-		ZenPakStreamingRequiresDeployAndLaunch,
-
-		Count
-	};
-}
-
-LAUNCHERSERVICES_API FString LexToStringLocalized(ELauncherProfileValidationErrors::Type Value);
 
 /** Type definition for shared pointers to instances of ILauncherProfile. */
 typedef TSharedPtr<class ILauncherSimpleProfile> ILauncherSimpleProfilePtr;
@@ -400,8 +147,10 @@ DECLARE_MULTICAST_DELEGATE(FOnProfileProjectChanged);
 /** Delegate type for a change in build target options */
 DECLARE_MULTICAST_DELEGATE(FOnProfileBuildTargetOptionsChanged);
 
-/** Delegate type for gathering custom profile validation */
-DECLARE_MULTICAST_DELEGATE(FOnProfileCustomValidation);
+
+/** Delegate type for a UAT command */
+DECLARE_MULTICAST_DELEGATE_OneParam(FOnLauncherProfileUATCommandChanged, const ILauncherProfileUATCommandRef&)
+
 
 
 /** Delegate type for detecting if cook is finished
@@ -460,17 +209,21 @@ public:
 	/**
 	 * Checks whether the last validation yielded any error.
 	 *
+   	 * @param UATCommand Optional UAT command associated with this error
 	 * @return true if the any error is present, false otherwise.
 	 */
-	virtual bool HasValidationError() const = 0;
+	virtual bool HasValidationError(ILauncherProfileUATCommandPtr UATCommand = nullptr) const = 0;
 
 	/**
 	 * Checks whether the last validation yielded the specified error.
 	 *
 	 * @param Error The validation error to check for.
+   	 * @param UATCommand Optional UAT command associated with this error
 	 * @return true if the error is present, false otherwise.
 	 */
-	virtual bool HasValidationError( ELauncherProfileValidationErrors::Type Error ) const = 0;
+	virtual bool HasValidationError( ELauncherProfileValidationErrors::Type Error, ILauncherProfileUATCommandPtr UATCommand = nullptr ) const = 0;
+
+
 
 
 	/**
@@ -478,8 +231,9 @@ public:
 	 * 
 	 * @param UniqueId Unique identifier for this error
 	 * @param Text Localized string describing this error
+ 	 * @param UATCommand Optional UAT command associated with this error
 	 */
-	virtual void AddCustomError( const FString& UniqueId, const FText& Text ) = 0;
+	virtual void AddCustomError( const FString& UniqueId, const FText& Text, ILauncherProfileUATCommandPtr UATCommand = nullptr ) = 0;
 
 
 	/**
@@ -487,7 +241,7 @@ public:
 	 *
 	 * @param UniqueId Unique identifier for this error
 	 */
-	virtual bool HasCustomError( const FString& UniqueId ) const = 0;
+	virtual bool HasCustomError( const FString& UniqueId, ILauncherProfileUATCommandPtr UATCommand = nullptr ) const = 0;
 
 	/**
 	 * Gets the localized string representation of the given error
@@ -495,14 +249,14 @@ public:
 	 * @param UniqueId
 	 * @returns Localized string describing this error, or empty text if we don't have it
 	 */
-	virtual FText GetCustomErrorText( const FString& UniqueId ) const = 0;
+	virtual FText GetCustomErrorText( const FString& UniqueId, ILauncherProfileUATCommandPtr UATCommand = nullptr ) const = 0;
 
 	/**
 	 * Gets all custom errors in this profile
 	 * 
 	 * @returns UniqueIds for all custom errors in this profile
 	 */
-	virtual TArray<FString> GetAllCustomErrors() const = 0;
+	virtual TArray<FString> GetAllCustomErrors(ILauncherProfileUATCommandPtr UATCommand = nullptr) const = 0;
 	
 
 
@@ -511,15 +265,16 @@ public:
 	 * 
 	 * @param UniqueId Unique identifier for this warning
 	 * @param Text Localized string describing this warning
+ 	 * @param UATCommand Optional UAT command associated with this warning
 	 */
-	virtual void AddCustomWarning( const FString& UniqueId, const FText& Text ) = 0;
+	virtual void AddCustomWarning( const FString& UniqueId, const FText& Text, ILauncherProfileUATCommandPtr UATCommand = nullptr ) = 0;
 
 	/**
 	 * Checks whether this profile has the given custom warning
 	 *
 	 * @param UniqueId Unique identifier for this warning
 	 */
-	virtual bool HasCustomWarning( const FString& UniqueId ) const = 0;
+	virtual bool HasCustomWarning( const FString& UniqueId, ILauncherProfileUATCommandPtr UATCommand = nullptr ) const = 0;
 
 	/**
 	 * Gets the localized string representation of the given warning
@@ -527,20 +282,20 @@ public:
 	 * @param UniqueId
 	 * @returns Localized string describing this warning, or empty text if we don't have it
 	 */
-	virtual FText GetCustomWarningText( const FString& UniqueId ) const = 0;
+	virtual FText GetCustomWarningText( const FString& UniqueId, ILauncherProfileUATCommandPtr UATCommand = nullptr ) const = 0;
 
 	/**
 	 * Gets all custom warnings in this profile
 	 * 
 	 * @returns UniqueIds for all custom warnings in this profile
 	 */
-	virtual TArray<FString> GetAllCustomWarnings() const = 0;
+	virtual TArray<FString> GetAllCustomWarnings(ILauncherProfileUATCommandPtr UATCommand = nullptr) const = 0;
 
 
 	/**
 	 * Refreshes the custom warnings and errors. Resets the lists and triggers the OnValidation callback.
 	 */
-	virtual void RefreshCustomWarningsAndErrors() = 0;
+	virtual void RefreshCustomWarningsAndErrors(ILauncherProfileUATCommandPtr UATCommand = nullptr) = 0;
 
 
 
@@ -549,7 +304,7 @@ public:
 	 *
 	 * @return string specifying the invalid platform.
 	 */
-	virtual FString GetInvalidPlatform() const = 0;
+	virtual FString GetInvalidPlatform(ILauncherProfileUATCommandPtr UATCommand = nullptr) const = 0;
 
 	/**
 	 * Checks whether devices of the specified platform can be deployed to.
@@ -1518,6 +1273,13 @@ public:
 	 */
 	virtual FOnProfileCustomValidation& OnCustomValidation() = 0;
 
+	/**
+	 * Access delegate used when validation changes to gather custom warnings and errors for custom UAT commands
+	 * 
+	 * @return The delegate.
+	 */
+	virtual FOnProfileCustomUATCommandValidation& OnCustomUATCommandValidation() = 0;
+
 	/**
 	 * Sets whether to use I/O store for optimized loading.
 	 * @param bUseIoStore Whether to use I/O store
@@ -1682,6 +1444,37 @@ public:
 	virtual void RemoveUATCommand( const TCHAR* InternalName ) = 0;
 
 
+	/**
+	 * Callback for when a UAT command is added
+	 */
+	virtual FOnLauncherProfileUATCommandChanged& OnUATCommandAdded() = 0;
+
+	/**
+	 * Callback for when a UAT command is removed, prior to deletion
+	 */
+	virtual FOnLauncherProfileUATCommandChanged& OnUATCommandRemoved() = 0;
+
+
+
+
+
+
+
+	/**
+	 * Returns a list of all build cook run commands in the profile
+	 */
+	virtual const TArray<ILauncherProfileBuildCookRunRef> GetBuildCookRunCommands() const = 0;
+
+	/**
+	 * Return an existing UAT build cook run command by internal name
+	 */
+	virtual ILauncherProfileBuildCookRunPtr GetBuildCookRunCommand( const TCHAR* InternalName ) const = 0;
+
+	/** 
+	 * Returns a UAT build cook run command by internal name, creating it if it doesn't already exist
+	 */
+	virtual ILauncherProfileBuildCookRunRef FindOrAddBuildCookRunCommand( const TCHAR* InternalName ) = 0;
+
 
 
 
@@ -1781,6 +1574,9 @@ public:
 	/** Whether this profile can be used by the legacy Project Launcher */
 	virtual bool SupportsLegacyProjectLauncher() const = 0;
 
+	/** Get access to the first build cook run command, if any */
+	virtual ILauncherProfileBuildCookRunPtr GetFirstBuildCookRun() const = 0;
+
 	/** Virtual destructor. */
 	virtual ~ILauncherProfile( ) = default;
 };
diff --git a/Engine/Source/Developer/LauncherServices/Public/ILauncherProfileAutomatedTest.h b/Engine/Source/Developer/LauncherServices/Public/ILauncherProfileAutomatedTest.h
--- a/Engine/Source/Developer/LauncherServices/Public/ILauncherProfileAutomatedTest.h
+++ b/Engine/Source/Developer/LauncherServices/Public/ILauncherProfileAutomatedTest.h
@@ -18,7 +18,7 @@ public:
 	virtual void SetTests( const TCHAR* Tests ) = 0; //-Test=xxx,yyy
 	virtual const FString& GetTests() const = 0;
 
-	virtual void Load(const FJsonObject& Object) = 0;
+	virtual bool Load(const FJsonObject& Object) = 0;
 	virtual void Save(TJsonWriter<>& Writer) = 0;
 
 	UE_DEPRECATED(5.8, "please use the Json Save/Load API")
diff --git a/Engine/Source/Developer/LauncherServices/Public/ILauncherProfileUATCommand.h b/Engine/Source/Developer/LauncherServices/Public/ILauncherProfileUATCommand.h
--- a/Engine/Source/Developer/LauncherServices/Public/ILauncherProfileUATCommand.h
+++ b/Engine/Source/Developer/LauncherServices/Public/ILauncherProfileUATCommand.h
@@ -7,6 +7,7 @@
 
 class FJsonObject;
 class ILauncherProfileAutomatedTest;
+class ILauncherProfileBuildCookRun;
 
 /**
  * UAT command parameters
@@ -19,8 +20,8 @@ public:
 	virtual void SetUATCommand( const TCHAR* UATCommand ) = 0;
 	virtual const FString& GetUATCommand() const = 0;
 
-	virtual void SetAdditionalCommandLine( const TCHAR* AdditionalCommandLine ) = 0;
-	virtual const FString& GetAdditionalCommandLine() const = 0;
+	virtual void SetAdditionalUATCommandLine( const TCHAR* AdditionalCommandLine ) = 0;
+	virtual const FString& GetAdditionalUATCommandLine() const = 0;
 
 	virtual void SetEnabled( bool bEnabled ) = 0;
 	virtual bool IsEnabled() const = 0;
@@ -31,10 +32,14 @@ public:
 	// recommendation for InternalName is is "ExtensionName.InternalName"
 	virtual const TCHAR* GetInternalName() const = 0;
 
-	virtual void Load(const FJsonObject& Object) = 0;
+	virtual void SetDescription(const FString& NewDescription) = 0;
+	virtual FString GetDescription() const = 0;
+
+	virtual bool Load(const FJsonObject& Object) = 0;
 	virtual void Save(TJsonWriter<>& Writer) = 0;
 
 	virtual TSharedPtr<ILauncherProfileAutomatedTest> AsAutomatedTest() = 0;
+	virtual TSharedPtr<ILauncherProfileBuildCookRun> AsBuildCookRun() = 0;
 
 	virtual ~ILauncherProfileUATCommand() = default;
 };
