fix for some UAT parameters being passed globally, instead of per-BuildCookRun

#rnx
#Synced-CL 49676362

diff --git a/Engine/Source/Developer/LauncherServices/Private/Launcher/LauncherUATTask.h b/Engine/Source/Developer/LauncherServices/Private/Launcher/LauncherUATTask.h
--- a/Engine/Source/Developer/LauncherServices/Private/Launcher/LauncherUATTask.h
+++ b/Engine/Source/Developer/LauncherServices/Private/Launcher/LauncherUATTask.h
@@ -17,13 +17,12 @@ class FLauncherUATTask
 {
 public:
 
-	FLauncherUATTask( const FString& InCommandLine, const FString& InName, const FString& InDesc, void* InReadPipe, void* InWritePipe, const FString& InEditorExe, FProcHandle& InProcessHandle, ILauncherWorker* InWorker, const FString& InCommandEnd)
+	FLauncherUATTask( const FString& InCommandLine, const FString& InName, const FString& InDesc, void* InReadPipe, void* InWritePipe, FProcHandle& InProcessHandle, ILauncherWorker* InWorker, const FString& InCommandEnd)
 		: FLauncherTask(InName, InDesc)
 
 		, CommandLine(InCommandLine)
 		, ReadPipe(InReadPipe)
 		, WritePipe(InWritePipe)
-		, EditorExe(InEditorExe)
 		, ProcessHandle(InProcessHandle)
 		, CommandText(InCommandEnd)
 	{
@@ -62,17 +61,8 @@ protected:
 		// Only allow UAT to compile scripts dynamically if we pass -development or we have the IsBuildingUAT property set, the latter of which should not allow itself to be set in installed situations.
 		bool bAllowCompile = true;
 		UATCommandLine += (bAllowCompile && (FParse::Param( FCommandLine::Get(), TEXT("development") ) || ChainState.Profile->IsBuildingUAT())) ? TEXT("") : TEXT(" -nocompile");
-		// we never want to build the editor when launching from the editor or running with an installed engine (which can't rebuild itself)
-		UATCommandLine += GIsEditor || FApp::IsEngineInstalled() ? TEXT(" -nocompileeditor") : TEXT("");
-		UATCommandLine += FApp::IsEngineInstalled() ? TEXT(" -installed") : TEXT("");
 		UATCommandLine += TEXT(" -utf8output");
 
-		// specify the path to the editor exe if necessary
-		if(EditorExe.Len() > 0)
-		{
-			UATCommandLine += FString::Printf(TEXT(" -unrealexe=\"%s\""), *EditorExe);
-		}
-
 		// specialized command arguments for this particular task
 		UATCommandLine += TEXT(" ");
 		UATCommandLine += CommandLine;
@@ -117,9 +107,6 @@ private:
 	void* ReadPipe = nullptr;
 	void* WritePipe = nullptr;
 
-	// The editor executable that UAT should use
-	FString EditorExe;
-
 	// process handle
 	FProcHandle& ProcessHandle;
 
diff --git a/Engine/Source/Developer/LauncherServices/Private/Launcher/LauncherWorker.cpp b/Engine/Source/Developer/LauncherServices/Private/Launcher/LauncherWorker.cpp
--- a/Engine/Source/Developer/LauncherServices/Private/Launcher/LauncherWorker.cpp
+++ b/Engine/Source/Developer/LauncherServices/Private/Launcher/LauncherWorker.cpp
@@ -408,6 +408,17 @@ FString FLauncherWorker::CreateUATCommandForBuildCookRun( const ILauncherProfile
 		LexToString(InBuildCookRun->GetBuildConfiguration()),
 		LexToString(InBuildCookRun->GetBuildConfiguration()));
 
+	// we never want to build the editor when launching from the editor or running with an installed engine (which can't rebuild itself)
+	UATCommand += GIsEditor || FApp::IsEngineInstalled() ? TEXT(" -nocompileeditor") : TEXT("");
+	UATCommand += FApp::IsEngineInstalled() ? TEXT(" -installed") : TEXT("");
+
+	// specify the path to the editor exe if necessary
+	FString EditorExe =  InProfile->GetEditorExe();
+	if(EditorExe.Len() > 0)
+	{
+		UATCommand += FString::Printf(TEXT(" -unrealexe=\"%s\""), *EditorExe);
+	}
+
 	FGuid SessionId(FGuid::NewGuid());
 	FString InitialMap = InBuildCookRun->GetDefaultLaunchRole()->GetInitialMap();
 	if (InitialMap.IsEmpty() && InBuildCookRun->GetCookedMaps().Num() == 1)
@@ -1769,7 +1780,7 @@ void FLauncherWorker::CreateAndExecuteTasks( const ILauncherProfileRef& InProfil
 	}
 
 
-	TSharedPtr<FLauncherTask> LaunchTask = MakeShareable(new FLauncherUATTask(UATCommandString, TEXT("Launch Task"), TEXT("Launching UAT..."), ReadPipe, WritePipe, InProfile->GetEditorExe(), ProcHandle, this, StartString));
+	TSharedPtr<FLauncherTask> LaunchTask = MakeShareable(new FLauncherUATTask(UATCommandString, TEXT("Launch Task"), TEXT("Launching UAT..."), ReadPipe, WritePipe, ProcHandle, this, StartString));
 	AddTask(LaunchTask);
 	for (int32 Index = 0; Index < Commands.Num(); ++Index)
 	{
