use absolute project path for automated test uat commands, to match build cook run

#rnx

#Synced-CL 49921272

diff --git a/Engine/Source/Developer/LauncherServices/Private/Launcher/LauncherWorker.cpp b/Engine/Source/Developer/LauncherServices/Private/Launcher/LauncherWorker.cpp
--- a/Engine/Source/Developer/LauncherServices/Private/Launcher/LauncherWorker.cpp
+++ b/Engine/Source/Developer/LauncherServices/Private/Launcher/LauncherWorker.cpp
@@ -1238,7 +1238,7 @@ FString FLauncherWorker::CreateUATCommandForAutomatedTest( const ILauncherProfil
 
 	// basic command parameters
 	FString AutomatedTestCommand = FString::Printf( TEXT("%s -utf8output -nop4 "), *InAutomatedTest->GetUATCommand() );
-	AutomatedTestCommand += FString::Printf(TEXT("-project=\"%s\" "), *InProfile->GetProjectPath());
+	AutomatedTestCommand += FString::Printf(TEXT("-project=\"%s\" "), *FPaths::ConvertRelativePathToFull(InProfile->GetProjectPath()));
 	AutomatedTestCommand += FString::Printf(TEXT("-configuration=%s "), LexToString(InProfile->GetBuildConfiguration()) );
 	AutomatedTestCommand += FString::Printf(TEXT("-tests=%s "), *InAutomatedTest->GetTests());
 
