UATCommands re-run profile validation when a property is changed
- also add missing property change notification in uat command toggle in PL2

#rnx

#Synced-CL 50155938

diff --git a/Engine/Plugins/Developer/ProjectLauncher/Source/ProjectLauncher/Private/Widgets/Profiles/SCustomLaunchCustomProfileEditor.cpp b/Engine/Plugins/Developer/ProjectLauncher/Source/ProjectLauncher/Private/Widgets/Profiles/SCustomLaunchCustomProfileEditor.cpp
--- a/Engine/Plugins/Developer/ProjectLauncher/Source/ProjectLauncher/Private/Widgets/Profiles/SCustomLaunchCustomProfileEditor.cpp
+++ b/Engine/Plugins/Developer/ProjectLauncher/Source/ProjectLauncher/Private/Widgets/Profiles/SCustomLaunchCustomProfileEditor.cpp
@@ -1223,7 +1223,11 @@ void SCustomLaunchCustomProfileEditor::MakeUATCommandSubMenu( FMenuBuilder& Menu
 		FText::GetEmpty(),
 		FSlateIcon(),
 		FUIAction(
-			FExecuteAction::CreateLambda( [UATCommand]() { UATCommand->SetEnabled( !UATCommand->IsEnabled() ); } ),
+			FExecuteAction::CreateLambda( [this, UATCommand]()
+			{
+				UATCommand->SetEnabled( !UATCommand->IsEnabled() ); 
+				TreeBuilder->OnPropertyChanged();
+			}),
 			FCanExecuteAction(),
 			FIsActionChecked::CreateLambda( [UATCommand]() { return UATCommand->IsEnabled(); } )
 		),
diff --git a/Engine/Source/Developer/LauncherServices/Private/Profiles/LauncherProfile.cpp b/Engine/Source/Developer/LauncherServices/Private/Profiles/LauncherProfile.cpp
--- a/Engine/Source/Developer/LauncherServices/Private/Profiles/LauncherProfile.cpp
+++ b/Engine/Source/Developer/LauncherServices/Private/Profiles/LauncherProfile.cpp
@@ -16,13 +16,13 @@ const TCHAR* FLauncherProfileBuildCookRun::TypeName = TEXT("BuildCookRun");
 
 ILauncherProfileUATCommandRef FLauncherProfile::CreateUATCommandInternal(const TCHAR* InternalName, const TCHAR* UserTypeName)
 {
-	ILauncherProfileUATCommandRef UATCommand = MakeShared<FLauncherProfileUATCommand>(InternalName, UserTypeName);
+	ILauncherProfileUATCommandRef UATCommand = MakeShared<FLauncherProfileUATCommand>(this, InternalName, UserTypeName);
 	return UATCommand;
 }
 
 ILauncherProfileAutomatedTestRef FLauncherProfile::CreateAutomatedTestInternal(const TCHAR* InternalName, const TCHAR* UserTypeName)
 {
-	ILauncherProfileAutomatedTestRef AutomatedTest = MakeShared<FLauncherProfileAutomatedTest>(InternalName, UserTypeName);
+	ILauncherProfileAutomatedTestRef AutomatedTest = MakeShared<FLauncherProfileAutomatedTest>(this, InternalName, UserTypeName);
 	return AutomatedTest;
 }
 
diff --git a/Engine/Source/Developer/LauncherServices/Private/Profiles/LauncherProfile.h b/Engine/Source/Developer/LauncherServices/Private/Profiles/LauncherProfile.h
--- a/Engine/Source/Developer/LauncherServices/Private/Profiles/LauncherProfile.h
+++ b/Engine/Source/Developer/LauncherServices/Private/Profiles/LauncherProfile.h
@@ -2511,18 +2511,18 @@ protected:
 		FValidationData& ValidationData = UATCommandValidation.FindChecked(BuildCookRun.ToSharedPtr());
 		ValidationData.Reset();
 
-		// don't verify if this build command is disabled - it won't affect the build
-		if (!BuildCookRun->IsEnabled())
-		{
-			return;
-		}
-
 		// temp. fix : merge in the global validation errors so legacy PL can display them
 		if (BuildCookRun == FirstBuildCookRun)
 		{
 			ValidationData.Errors = ProfileValidation.Errors;
 		}
 
+		// don't verify if this build command is disabled - it won't affect the build
+		if (!BuildCookRun->IsEnabled())
+		{
+			return;
+		}
+
 		// Build: a build configuration must be selected
 		if (BuildCookRun->GetBuildConfiguration() == EBuildConfiguration::Unknown)
 		{
@@ -3068,5 +3068,8 @@ private:
 	typedef TMap<FString, int32> FPlatformToBuildsMap;
 	static TMap<FString, FPlatformToBuildsMap> PerProjectBuilds;
 
-	friend class FLauncherProfileBuildCookRun; // for legacy serialize
+	// friendship so UAT command specializations can access Validate()
+	friend class FLauncherProfileBuildCookRun; // (also for legacy serialize)
+	friend class FLauncherProfileAutomatedTest;
+	friend class FLauncherProfileUATCommand;
 };
diff --git a/Engine/Source/Developer/LauncherServices/Private/Profiles/LauncherProfileAutomatedTest.h b/Engine/Source/Developer/LauncherServices/Private/Profiles/LauncherProfileAutomatedTest.h
--- a/Engine/Source/Developer/LauncherServices/Private/Profiles/LauncherProfileAutomatedTest.h
+++ b/Engine/Source/Developer/LauncherServices/Private/Profiles/LauncherProfileAutomatedTest.h
@@ -5,6 +5,7 @@
 #include "CoreMinimal.h"
 #include "ILauncherProfileAutomatedTest.h"
 #include "ILauncherProfileBuildCookRun.h"
+#include "Profiles/LauncherProfile.h"
 #include "Dom/JsonObject.h"
 #include "Internationalization/Internationalization.h"
 
@@ -13,8 +14,9 @@ class FLauncherProfileAutomatedTest final : public ILauncherProfileAutomatedTest
 public:
 	static const TCHAR* TypeName;
 
-	FLauncherProfileAutomatedTest( const TCHAR* InInternalName = TEXT(""), const TCHAR* InUserTypeName = TEXT("") ) 
-		: InternalName(InInternalName)
+	FLauncherProfileAutomatedTest( FLauncherProfile* InProfile, const TCHAR* InInternalName = TEXT(""), const TCHAR* InUserTypeName = TEXT("") ) 
+		: Profile(InProfile)
+		, InternalName(InInternalName)
 		, UserTypeName(InUserTypeName)
 	{
 	}
@@ -30,6 +32,10 @@ public:
 	virtual void SetUATCommand( const TCHAR* InUATCommand ) override
 	{
 		UATCommand = InUATCommand;
+		if (Profile)
+		{
+			Profile->Validate();
+		}
 	}
 
 	virtual const FString& GetUATCommand() const override
@@ -40,6 +46,10 @@ public:
 	virtual void SetTests( const TCHAR* InTests ) override
 	{
 		Tests = InTests;
+		if (Profile)
+		{
+			Profile->Validate();
+		}
 	}
 	virtual const FString& GetTests() const override
 	{
@@ -49,6 +59,10 @@ public:
 	virtual void SetAdditionalUATCommandLine( const TCHAR* InAdditionalCommandLine ) override
 	{
 		UATCommandLine = InAdditionalCommandLine;
+		if (Profile)
+		{
+			Profile->Validate();
+		}
 	}
 	virtual const FString& GetAdditionalUATCommandLine() const override
 	{
@@ -58,6 +72,10 @@ public:
 	virtual void SetEnabled( bool bInEnabled ) override
 	{
 		bEnabled = bInEnabled;
+		if (Profile)
+		{
+			Profile->Validate();
+		}
 	}
 	virtual bool IsEnabled() const override
 	{
@@ -67,6 +85,10 @@ public:
 	virtual void SetOrder( int32 InOrder ) override
 	{
 		Order = InOrder;
+		if (Profile)
+		{
+			Profile->Validate();
+		}
 	}
 	virtual int32 GetOrder() const override
 	{
@@ -81,6 +103,10 @@ public:
 	virtual void SetDescription(const FString& NewDescription) override
 	{
 		Description = NewDescription;
+		if (Profile)
+		{
+			Profile->Validate();
+		}
 	}
 
 	virtual FString GetDescription() const
@@ -91,6 +117,10 @@ public:
 	virtual void SetUserTypeName(const FString& NewUserTypeName) override
 	{
 		UserTypeName = NewUserTypeName;
+		if (Profile)
+		{
+			Profile->Validate();
+		}
 	}
 
 	virtual FString GetUserTypeName() const override
@@ -187,7 +217,8 @@ public:
 	}
 
 
-protected:
+private:
+	FLauncherProfile* Profile;
 
 	FString UATCommand = TEXT("RunUnreal");
 	FString Tests; //-Test=xxx,yyy
diff --git a/Engine/Source/Developer/LauncherServices/Private/Profiles/LauncherProfileBuildCookRun.h b/Engine/Source/Developer/LauncherServices/Private/Profiles/LauncherProfileBuildCookRun.h
--- a/Engine/Source/Developer/LauncherServices/Private/Profiles/LauncherProfileBuildCookRun.h
+++ b/Engine/Source/Developer/LauncherServices/Private/Profiles/LauncherProfileBuildCookRun.h
@@ -1789,6 +1789,10 @@ public:
 	virtual void SetUATCommand( const TCHAR* InUATCommand ) override
 	{
 		UATCommand = InUATCommand;
+		if (Profile)
+		{
+			Profile->Validate();
+		}
 	}
 
 	virtual const FString& GetUATCommand() const override
@@ -1799,6 +1803,10 @@ public:
 	virtual void SetAdditionalUATCommandLine( const TCHAR* InAdditionalCommandLine ) override
 	{
 		UATCommandLine = InAdditionalCommandLine;
+		if (Profile)
+		{
+			Profile->Validate();
+		}
 	}
 	virtual const FString& GetAdditionalUATCommandLine() const override
 	{
@@ -1808,7 +1816,10 @@ public:
 	virtual void SetEnabled( bool bInEnabled ) override
 	{
 		bEnabled = bInEnabled;
-		Validate();
+		if (Profile)
+		{
+			Profile->Validate();
+		}
 	}
 	virtual bool IsEnabled() const override
 	{
@@ -1818,6 +1829,10 @@ public:
 	virtual void SetOrder( int32 InOrder ) override
 	{
 		Order = InOrder;
+		if (Profile)
+		{
+			Profile->Validate();
+		}
 	}
 	virtual int32 GetOrder() const override
 	{
@@ -1832,6 +1847,10 @@ public:
 	virtual void SetUserTypeName(const FString& NewUserTypeName) override
 	{
 		UserTypeName = NewUserTypeName;
+		if (Profile)
+		{
+			Profile->Validate();
+		}
 	}
 
 	virtual FString GetUserTypeName() const override
diff --git a/Engine/Source/Developer/LauncherServices/Private/Profiles/LauncherProfileUATCommand.h b/Engine/Source/Developer/LauncherServices/Private/Profiles/LauncherProfileUATCommand.h
--- a/Engine/Source/Developer/LauncherServices/Private/Profiles/LauncherProfileUATCommand.h
+++ b/Engine/Source/Developer/LauncherServices/Private/Profiles/LauncherProfileUATCommand.h
@@ -6,6 +6,7 @@
 #include "ILauncherProfileUATCommand.h"
 #include "ILauncherProfileAutomatedTest.h"
 #include "ILauncherProfileBuildCookRun.h"
+#include "Profiles/LauncherProfile.h"
 #include "Dom/JsonObject.h"
 #include "Internationalization/Internationalization.h"
 
@@ -14,8 +15,9 @@ class FLauncherProfileUATCommand final : public ILauncherProfileUATCommand, publ
 public:
 	static const TCHAR* TypeName;
 
-	FLauncherProfileUATCommand(const TCHAR* InInternalName = TEXT(""), const TCHAR* InUserTypeName = TEXT("") ) 
-		: InternalName(InInternalName)
+	FLauncherProfileUATCommand(FLauncherProfile* InProfile, const TCHAR* InInternalName = TEXT(""), const TCHAR* InUserTypeName = TEXT("") ) 
+		: Profile(InProfile)
+		, InternalName(InInternalName)
 		, UserTypeName(InUserTypeName)
 	{
 	}
@@ -23,6 +25,10 @@ public:
 	virtual void SetUATCommand( const TCHAR* InUATCommand ) override
 	{
 		UATCommand = InUATCommand;
+		if (Profile)
+		{
+			Profile->Validate();
+		}
 	}
 
 	virtual const FString& GetUATCommand() const override
@@ -33,6 +39,10 @@ public:
 	virtual void SetAdditionalUATCommandLine( const TCHAR* InAdditionalCommandLine ) override
 	{
 		UATCommandLine = InAdditionalCommandLine;
+		if (Profile)
+		{
+			Profile->Validate();
+		}
 	}
 	virtual const FString& GetAdditionalUATCommandLine() const override
 	{
@@ -42,6 +52,10 @@ public:
 	virtual void SetEnabled( bool bInEnabled ) override
 	{
 		bEnabled = bInEnabled;
+		if (Profile)
+		{
+			Profile->Validate();
+		}
 	}
 	virtual bool IsEnabled() const override
 	{
@@ -51,6 +65,10 @@ public:
 	virtual void SetOrder( int32 InOrder ) override
 	{
 		Order = InOrder;
+		if (Profile)
+		{
+			Profile->Validate();
+		}
 	}
 	virtual int32 GetOrder() const override
 	{
@@ -65,6 +83,10 @@ public:
 	virtual void SetDescription(const FString& NewDescription) override
 	{
 		Description = NewDescription;
+		if (Profile)
+		{
+			Profile->Validate();
+		}
 	}
 
 	virtual FString GetDescription() const override
@@ -75,6 +97,10 @@ public:
 	virtual void SetUserTypeName(const FString& NewUserTypeName) override
 	{
 		UserTypeName = NewUserTypeName;
+		if (Profile)
+		{
+			Profile->Validate();
+		}
 	}
 
 	virtual FString GetUserTypeName() const override
@@ -150,7 +176,8 @@ public:
 	}
 
 
-protected:
+private:
+	FLauncherProfile* Profile;
 
 	FString UATCommand = TEXT("RunUnreal");
 	FString UATCommandLine;
