Project launcher automated tests require at least one build cook run step
- alternative is to duplicate required properties into every automated test extension which seems excessive at present

#rnx
#Synced-CL 50155938

diff --git a/Engine/Source/Developer/LauncherServices/Private/Profiles/LauncherProfile.cpp b/Engine/Source/Developer/LauncherServices/Private/Profiles/LauncherProfile.cpp
--- a/Engine/Source/Developer/LauncherServices/Private/Profiles/LauncherProfile.cpp
+++ b/Engine/Source/Developer/LauncherServices/Private/Profiles/LauncherProfile.cpp
@@ -75,7 +75,7 @@ ILauncherProfileBuildCookRunPtr FLauncherProfile::LoadBuildCookRunInternal(const
 
 FString LexToStringLocalized(ELauncherProfileValidationErrors::Type Value)
 {
-	static_assert(ELauncherProfileValidationErrors::Count == 31, "GetLocalizedValidationErrorMessage() needs to be updated to account for modified enum values.");
+	static_assert(ELauncherProfileValidationErrors::Count == 32, "GetLocalizedValidationErrorMessage() needs to be updated to account for modified enum values.");
 	switch (Value)
 	{
 		case ELauncherProfileValidationErrors::CopyToDeviceRequiresCookByTheBook:
@@ -140,6 +140,8 @@ FString LexToStringLocalized(ELauncherProfileValidationErrors::Type Value)
 			return LOCTEXT("CopyToDeviceRequiresNoPackaging", "Deployment by copying to device requires 'Do not package' packaging.").ToString();
 		case ELauncherProfileValidationErrors::ZenPakStreamingRequiresDeployAndLaunch:
 			return LOCTEXT("ZenPakStreamingRequiresDeployAndLaunch", "Zen Pak Streaming requires deployment to a device and launching.").ToString();
+		case ELauncherProfileValidationErrors::AutomatedTestRequiredBuildCookRun:
+			return LOCTEXT("AutomatedTestRequiredBuildCookRun", "Automated Tests currently require at least one active BuildCookRun.").ToString();
 		default:
 			return TEXT("Unknown");
 	};
diff --git a/Engine/Source/Developer/LauncherServices/Private/Profiles/LauncherProfile.h b/Engine/Source/Developer/LauncherServices/Private/Profiles/LauncherProfile.h
--- a/Engine/Source/Developer/LauncherServices/Private/Profiles/LauncherProfile.h
+++ b/Engine/Source/Developer/LauncherServices/Private/Profiles/LauncherProfile.h
@@ -2489,13 +2489,31 @@ protected:
 		ProfileValidation.Reset();
 		RefreshCustomWarningsAndErrors(nullptr);
 
+		TArray<ILauncherProfileBuildCookRunRef> BuildCookRuns = GetBuildCookRunCommands();
+
 		// Build: a project must be selected
 		if (GetProjectPath().IsEmpty())
 		{
 			ProfileValidation.Errors.Add(ELauncherProfileValidationErrors::NoProjectSelected);
 		}
 
-		for (const ILauncherProfileBuildCookRunRef& BuildCookRun : GetBuildCookRunCommands())
+		// Must have an active BuildCookRun to use automated tests (@todo: eventually could add relevent data into each automated test, but maybe an overkill?)
+		bool bHasActiveBuildCookRun = false;
+		for (const ILauncherProfileBuildCookRunRef& BuildCookRun : BuildCookRuns)
+		{
+			bHasActiveBuildCookRun |= BuildCookRun->IsEnabled();
+		}
+		for (const ILauncherProfileAutomatedTestRef& AutomatedTest : GetAutomatedTests())
+		{
+			if (!bHasActiveBuildCookRun && AutomatedTest->IsEnabled())
+			{
+				ProfileValidation.Errors.Add(ELauncherProfileValidationErrors::AutomatedTestRequiredBuildCookRun);
+				break;
+			}
+		}
+
+		// validate build cook runs
+		for (const ILauncherProfileBuildCookRunRef& BuildCookRun : BuildCookRuns)
 		{
 			Validate(BuildCookRun);
 		}
diff --git a/Engine/Source/Developer/LauncherServices/Public/ILauncherProfileBuildCookRun.h b/Engine/Source/Developer/LauncherServices/Public/ILauncherProfileBuildCookRun.h
--- a/Engine/Source/Developer/LauncherServices/Public/ILauncherProfileBuildCookRun.h
+++ b/Engine/Source/Developer/LauncherServices/Public/ILauncherProfileBuildCookRun.h
@@ -262,6 +262,9 @@ namespace ELauncherProfileValidationErrors
 		/** Zen Pak Streaming requires deployment and launching */
 		ZenPakStreamingRequiresDeployAndLaunch,
 
+		/** Automated Tests currently require at least one active BuildCookRun */
+		AutomatedTestRequiredBuildCookRun,
+
 		Count
 	};
 }
