Launch profile items that moved to the inner BuildCookRun are marked as deprecated
- fixup everywhere in project launcher, launcher worker and other call sites that were using the deprecated functions
- legacy project launcher & profile wizard code is muting the deprecation warnings for now

#rnx

#Synced-CL 50244055

diff --git a/Engine/Plugins/Developer/ProjectLauncher/Source/CommonLaunchExtensions/Private/AdvancedCook/AdvancedCookLaunchExtension.cpp b/Engine/Plugins/Developer/ProjectLauncher/Source/CommonLaunchExtensions/Private/AdvancedCook/AdvancedCookLaunchExtension.cpp
--- a/Engine/Plugins/Developer/ProjectLauncher/Source/CommonLaunchExtensions/Private/AdvancedCook/AdvancedCookLaunchExtension.cpp
+++ b/Engine/Plugins/Developer/ProjectLauncher/Source/CommonLaunchExtensions/Private/AdvancedCook/AdvancedCookLaunchExtension.cpp
@@ -47,9 +47,9 @@ bool FAdvancedCookLaunchExtensionInstance::CanToggleBuildCookRunExtension( const
 FAdvancedCookLaunchExtensionInstance::FBuildCookRunInstance::FBuildCookRunInstance( const FArgs& InArgs )
 	: ProjectLauncher::FBuildCookRunExtension( InArgs )
 {
-	MapOption = GetProfile()->GetCookedMaps().Num() > 0 ? EMapOption::Selected : EMapOption::Startup;
+	MapOption = GetBuildCookRun()->GetCookedMaps().Num() > 0 ? EMapOption::Selected : EMapOption::Startup;
 
-	CachedCulturesToCook = GetProfile()->GetCookedCultures();
+	CachedCulturesToCook = GetBuildCookRun()->GetCookedCultures();
 	CultureOption = (CachedCulturesToCook.Num() == 1 && CachedCulturesToCook.Contains(TEXT("en"))) ? ECultureOption::Default : ECultureOption::Custom;
 }
 
@@ -58,7 +58,7 @@ void FAdvancedCookLaunchExtensionInstance::FBuildCookRunInstance::CustomizeTree(
 {
 	auto IsVisible = [this]()
 	{
-		return (GetProfile()->GetCookMode() == ELauncherProfileCookModes::ByTheBook);
+		return (GetBuildCookRun()->GetCookMode() == ELauncherProfileCookModes::ByTheBook);
 	};
 
 	AddDefaultHeading(ProfileTreeNode)
diff --git a/Engine/Plugins/Developer/ProjectLauncher/Source/CommonLaunchExtensions/Private/Globals/GlobalsLaunchExtension.cpp b/Engine/Plugins/Developer/ProjectLauncher/Source/CommonLaunchExtensions/Private/Globals/GlobalsLaunchExtension.cpp
--- a/Engine/Plugins/Developer/ProjectLauncher/Source/CommonLaunchExtensions/Private/Globals/GlobalsLaunchExtension.cpp
+++ b/Engine/Plugins/Developer/ProjectLauncher/Source/CommonLaunchExtensions/Private/Globals/GlobalsLaunchExtension.cpp
@@ -12,9 +12,6 @@ const TCHAR* FGlobalsLaunchExtensionInstance::LocalHostVariable     = TEXT("$(Lo
 const TCHAR* FGlobalsLaunchExtensionInstance::LocalHostsVariable    = TEXT("$(LocalHosts)");
 const TCHAR* FGlobalsLaunchExtensionInstance::ProjectNameVariable   = TEXT("$(ProjectName)");
 const TCHAR* FGlobalsLaunchExtensionInstance::ProjectPathVariable   = TEXT("$(ProjectPath)");
-const TCHAR* FGlobalsLaunchExtensionInstance::TargetNameVariable    = TEXT("$(TargetName)");
-const TCHAR* FGlobalsLaunchExtensionInstance::PlatformNameVariable  = TEXT("$(Platform)");
-const TCHAR* FGlobalsLaunchExtensionInstance::ConfigurationVariable = TEXT("$(Configuration)");
 
 bool FGlobalsLaunchExtensionInstance::GetExtensionVariables( TArray<FString>& OutVariables ) const
 {
@@ -22,9 +19,6 @@ bool FGlobalsLaunchExtensionInstance::GetExtensionVariables( TArray<FString>& Ou
 	OutVariables.Add(LocalHostsVariable);
 	OutVariables.Add(ProjectNameVariable);
 	OutVariables.Add(ProjectPathVariable);
-	OutVariables.Add(TargetNameVariable);
-	OutVariables.Add(PlatformNameVariable);
-	OutVariables.Add(ConfigurationVariable);
 	return true;
 }
 
@@ -78,21 +72,6 @@ bool FGlobalsLaunchExtensionInstance::GetExtensionVariableValue( const FString&
 		OutValue = GetProfile()->GetProjectPath();
 		return true;
 	}
-	else if (InVariable == TargetNameVariable)
-	{
-		OutValue = FString::Join( GetProfile()->GetBuildTargets(), TEXT("+") );
-		return true;
-	}
-	else if (InVariable == PlatformNameVariable)
-	{
-		OutValue = FString::Join( GetProfile()->GetCookedPlatforms(), TEXT("+") );
-		return true;
-	}
-	else if (InVariable == ConfigurationVariable)
-	{
-		OutValue = LexToString( GetProfile()->GetBuildConfiguration() );
-		return true;
-	}
 
 	return false;
 }
diff --git a/Engine/Plugins/Developer/ProjectLauncher/Source/ProjectLauncher/Private/Extension/LaunchExtension.cpp b/Engine/Plugins/Developer/ProjectLauncher/Source/ProjectLauncher/Private/Extension/LaunchExtension.cpp
--- a/Engine/Plugins/Developer/ProjectLauncher/Source/ProjectLauncher/Private/Extension/LaunchExtension.cpp
+++ b/Engine/Plugins/Developer/ProjectLauncher/Source/ProjectLauncher/Private/Extension/LaunchExtension.cpp
@@ -599,13 +599,17 @@ namespace ProjectLauncher
 
 	FString FLaunchExtensionInstance::GetCommandLine() const
 	{
+		PRAGMA_DISABLE_DEPRECATION_WARNINGS
 		return Profile->GetAdditionalCommandLineParameters();
+		PRAGMA_ENABLE_DEPRECATION_WARNINGS
 	}
 
 	void FLaunchExtensionInstance::SetCommandLine( const FString& CommandLine )
 	{
+		PRAGMA_DISABLE_DEPRECATION_WARNINGS
 		Profile->SetAdditionalCommandLineParameters(CommandLine);
 		BroadcastPropertyChanged();
+		PRAGMA_ENABLE_DEPRECATION_WARNINGS
 	}
 
 	FName FLaunchExtensionInstance::GetLaunchPlatformName() const
diff --git a/Engine/Plugins/Developer/ProjectLauncher/Source/ProjectLauncher/Private/Model/ProjectLauncherModel.cpp b/Engine/Plugins/Developer/ProjectLauncher/Source/ProjectLauncher/Private/Model/ProjectLauncherModel.cpp
--- a/Engine/Plugins/Developer/ProjectLauncher/Source/ProjectLauncher/Private/Model/ProjectLauncherModel.cpp
+++ b/Engine/Plugins/Developer/ProjectLauncher/Source/ProjectLauncher/Private/Model/ProjectLauncherModel.cpp
@@ -68,7 +68,7 @@ namespace ProjectLauncher
 
 		// prepare profiles
 		BasicLaunchProfile = CreateBasicLaunchProfile();
-		bHasSetBasicLaunchProfilePlatform = (BasicLaunchProfile->GetDeployedDeviceGroup()->GetDeviceIDs().Num() > 0);
+		bHasSetBasicLaunchProfilePlatform = (BasicLaunchProfile->GetFirstBuildCookRun()->GetDeployedDeviceGroup()->GetDeviceIDs().Num() > 0);
 
 		AllProfiles = ProfileManager->GetAllProfiles();
 		for (const ILauncherProfilePtr& Profile : AllProfiles)
@@ -225,9 +225,9 @@ namespace ProjectLauncher
 
 	const PlatformInfo::FTargetPlatformInfo* FModel::GetPlatformInfo(const ILauncherProfilePtr& Profile)
 	{
-		if (Profile.IsValid() && Profile->GetCookedPlatforms().Num() > 0)
+		if (Profile.IsValid() && Profile->GetFirstBuildCookRun() && Profile->GetFirstBuildCookRun()->GetCookedPlatforms().Num() > 0)
 		{
-			FString SelectedPlatform = Profile->GetCookedPlatforms()[0];
+			FString SelectedPlatform = Profile->GetFirstBuildCookRun()->GetCookedPlatforms()[0];
 			return PlatformInfo::FindPlatformInfo(FName(SelectedPlatform));
 		}
 
@@ -306,9 +306,9 @@ namespace ProjectLauncher
 
 	bool FModel::IsHostPlatform(const ILauncherProfilePtr& Profile)
 	{
-		if (Profile.IsValid() && Profile->GetCookedPlatforms().Num() > 0)
+		if (Profile.IsValid() && Profile->GetFirstBuildCookRun() && Profile->GetFirstBuildCookRun()->GetCookedPlatforms().Num() > 0)
 		{
-			FString SelectedPlatform = Profile->GetCookedPlatforms()[0];
+			FString SelectedPlatform = Profile->GetFirstBuildCookRun()->GetCookedPlatforms()[0];
 			return IsHostPlatform(FName(SelectedPlatform));
 		}
 
@@ -490,7 +490,7 @@ PRAGMA_ENABLE_DEPRECATION_WARNINGS
 		if (!bHasSetBasicLaunchProfilePlatform && DeviceProxy->HasTargetPlatform(FPlatformProperties::IniPlatformName()) && ensure(BasicLaunchProfile->GetFirstBuildCookRun().IsValid()))
 		{
 			UpdatedCookedPlatformsFromDeployDeviceProxy(BasicLaunchProfile.ToSharedRef(), BasicLaunchProfile->GetFirstBuildCookRun().ToSharedRef(), DeviceProxy);
-			bHasSetBasicLaunchProfilePlatform = (BasicLaunchProfile->GetDeployedDeviceGroup()->GetDeviceIDs().Num() > 0);
+			bHasSetBasicLaunchProfilePlatform = (BasicLaunchProfile->GetFirstBuildCookRun()->GetDeployedDeviceGroup()->GetDeviceIDs().Num() > 0);
 		}
 	}
 
@@ -587,7 +587,7 @@ PRAGMA_ENABLE_DEPRECATION_WARNINGS
 		ILauncherProfileRef Profile = CreateCustomProfile(*BasicLaunchProfileName.ToString());
 		Profile->SetDescription(*BasicLaunchProfileDescription.ToString());
 		Profile->SetProjectSpecified(false);
-		Profile->SetBuildTargetSpecified(false);
+		Profile->GetFirstBuildCookRun()->SetBuildTargetSpecified(false);
 
 		return Profile;
 	}
@@ -745,11 +745,19 @@ PRAGMA_ENABLE_DEPRECATION_WARNINGS
 				NewProfile->AssignId(true);
 
 				// create a new device group because these are not serialized
-				ILauncherDeviceGroupRef DeployDeviceGroup = ProfileManager->AddNewDeviceGroup();
-				NewProfile->SetDeployedDeviceGroup(DeployDeviceGroup);
-				for (const FString& DeviceID : Profile->GetDeployedDeviceGroup()->GetDeviceIDs())
+				for (ILauncherProfileBuildCookRunRef NewBuildCookRun : NewProfile->GetBuildCookRunCommands())
 				{
-					DeployDeviceGroup->AddDevice(DeviceID);
+					ILauncherDeviceGroupRef DeployDeviceGroup = ProfileManager->AddNewDeviceGroup();
+					NewBuildCookRun->SetDeployedDeviceGroup(DeployDeviceGroup);
+
+					ILauncherProfileBuildCookRunPtr BuildCookRun = Profile->GetBuildCookRunCommand(NewBuildCookRun->GetInternalName());
+					if (BuildCookRun.IsValid())
+					{
+						for (const FString& DeviceID : BuildCookRun->GetDeployedDeviceGroup()->GetDeviceIDs())
+						{
+							DeployDeviceGroup->AddDevice(DeviceID);
+						}
+					}
 				}
 
 				return NewProfile;
diff --git a/Engine/Plugins/Developer/ProjectLauncher/Source/ProjectLauncher/Private/Widgets/Profiles/SCustomLaunchCustomProfileSelector.cpp b/Engine/Plugins/Developer/ProjectLauncher/Source/ProjectLauncher/Private/Widgets/Profiles/SCustomLaunchCustomProfileSelector.cpp
--- a/Engine/Plugins/Developer/ProjectLauncher/Source/ProjectLauncher/Private/Widgets/Profiles/SCustomLaunchCustomProfileSelector.cpp
+++ b/Engine/Plugins/Developer/ProjectLauncher/Source/ProjectLauncher/Private/Widgets/Profiles/SCustomLaunchCustomProfileSelector.cpp
@@ -526,10 +526,20 @@ void SCustomLaunchCustomProfileSelector::OnCustomProfileEditDescriptionClicked(I
 
 void SCustomLaunchCustomProfileSelector::OnDeviceRemoved( const FString DeviceID, ILauncherProfilePtr Profile)
 {
-	if (ensure(Profile->GetDeployedDeviceGroup() != nullptr) && Profile->GetDeployedDeviceGroup()->GetDeviceIDs().Contains(DeviceID))
+	//@todo: not sure whether this is needed?
+
+	bool bModified = false;
+	for (ILauncherProfileBuildCookRunRef BuildCookRun : Profile->GetBuildCookRunCommands())
 	{
-		Profile->GetDeployedDeviceGroup()->RemoveDevice(DeviceID);
+		if (ensure(BuildCookRun->GetDeployedDeviceGroup() != nullptr) && BuildCookRun->GetDeployedDeviceGroup()->GetDeviceIDs().Contains(DeviceID))
+		{
+			BuildCookRun->GetDeployedDeviceGroup()->RemoveDevice(DeviceID);
+			bModified = true;
+		}
+	}
 
+	if (bModified)
+	{
 		OnProfileModified.ExecuteIfBound(Profile);
 	}
 }
@@ -538,14 +548,17 @@ void SCustomLaunchCustomProfileSelector::SetSelectedDevices( const TArray<FStrin
 {
 	check(Model->IsBasicLaunchProfile(Profile));
 
-	if (ensure(Profile->GetDeployedDeviceGroup() != nullptr))
+	ILauncherProfileBuildCookRunPtr BuildCookRun = Profile->GetFirstBuildCookRun();
+	check(BuildCookRun.IsValid());
+
+	if (ensure(BuildCookRun->GetDeployedDeviceGroup() != nullptr))
 	{
-		Profile->GetDeployedDeviceGroup()->RemoveAllDevices();
+		BuildCookRun->GetDeployedDeviceGroup()->RemoveAllDevices();
 		for (const FString& DeviceID : DeviceIDs)
 		{
-			Profile->GetDeployedDeviceGroup()->AddDevice(DeviceID);
+			BuildCookRun->GetDeployedDeviceGroup()->AddDevice(DeviceID);
 		}
-		Model->UpdatedCookedPlatformsFromDeployDeviceProxy(Profile.ToSharedRef(), Profile->GetFirstBuildCookRun().ToSharedRef());
+		Model->UpdatedCookedPlatformsFromDeployDeviceProxy(Profile.ToSharedRef(), BuildCookRun.ToSharedRef());
 
 		OnProfileModified.ExecuteIfBound(Profile);
 	}
@@ -553,9 +566,14 @@ void SCustomLaunchCustomProfileSelector::SetSelectedDevices( const TArray<FStrin
 
 TArray<FString> SCustomLaunchCustomProfileSelector::GetSelectedDevices(const ILauncherProfilePtr Profile) const
 {
-	if (ensure(Profile->GetDeployedDeviceGroup() != nullptr))
+	check(Model->IsBasicLaunchProfile(Profile));
+
+	ILauncherProfileBuildCookRunPtr BuildCookRun = Profile->GetFirstBuildCookRun();
+	check(BuildCookRun.IsValid());
+
+	if (ensure(BuildCookRun->GetDeployedDeviceGroup() != nullptr))
 	{
-		return Profile->GetDeployedDeviceGroup()->GetDeviceIDs();
+		return BuildCookRun->GetDeployedDeviceGroup()->GetDeviceIDs();
 	}
 
 	return TArray<FString>();
diff --git a/Engine/Plugins/Developer/ProjectLauncher/Source/ProjectLauncher/Private/Widgets/Profiles/SCustomLaunchProfilesPanel.cpp b/Engine/Plugins/Developer/ProjectLauncher/Source/ProjectLauncher/Private/Widgets/Profiles/SCustomLaunchProfilesPanel.cpp
--- a/Engine/Plugins/Developer/ProjectLauncher/Source/ProjectLauncher/Private/Widgets/Profiles/SCustomLaunchProfilesPanel.cpp
+++ b/Engine/Plugins/Developer/ProjectLauncher/Source/ProjectLauncher/Private/Widgets/Profiles/SCustomLaunchProfilesPanel.cpp
@@ -702,27 +702,27 @@ const FSlateBrush* SCustomLaunchProfilesPanel::GetLaunchButtonImage() const
 		return FProjectLauncherStyle::Get().GetBrush("Icons.Task.Launch");
 	}	
 	// @fixme: figure out what the icon is like for multiple BCR steps
-	else if (Profile->GetLaunchMode() != ELauncherProfileLaunchModes::DoNotLaunch)
+	else if (Profile->GetFirstBuildCookRun()->GetLaunchMode() != ELauncherProfileLaunchModes::DoNotLaunch)
 	{
 		return FProjectLauncherStyle::Get().GetBrush("Icons.Task.Launch");
 	}
-	else if (Profile->GetDeploymentMode() != ELauncherProfileDeploymentModes::DoNotDeploy)
+	else if (Profile->GetFirstBuildCookRun()->GetDeploymentMode() != ELauncherProfileDeploymentModes::DoNotDeploy)
 	{
 		return FProjectLauncherStyle::Get().GetBrush("Icons.Task.Deploy");
 	}
-	else if (Profile->GetPackagingMode() != ELauncherProfilePackagingModes::DoNotPackage)
+	else if (Profile->GetFirstBuildCookRun()->GetPackagingMode() != ELauncherProfilePackagingModes::DoNotPackage)
 	{
 		return FProjectLauncherStyle::Get().GetBrush("Icons.Task.Package");
 	}
-	else if (Profile->GetCookMode() != ELauncherProfileCookModes::DoNotCook)
+	else if (Profile->GetFirstBuildCookRun()->GetCookMode() != ELauncherProfileCookModes::DoNotCook)
 	{
 		return FProjectLauncherStyle::Get().GetBrush("Icons.Task.Cook");
 	}
-	else if (Profile->GetBuildMode() != ELauncherProfileBuildModes::DoNotBuild)
+	else if (Profile->GetFirstBuildCookRun()->GetBuildMode() != ELauncherProfileBuildModes::DoNotBuild)
 	{
 		return FProjectLauncherStyle::Get().GetBrush("Icons.Task.Build");
 	}
-	else if (Profile->IsImportingZenSnapshot())
+	else if (Profile->GetFirstBuildCookRun()->IsImportingZenSnapshot())
 	{
 		return FProjectLauncherStyle::Get().GetBrush("Icons.Task.Zen");
 	}
diff --git a/Engine/Plugins/Editor/MobileLauncherProfileWizard/Source/MobileLauncherProfileWizard/Private/AndroidProfileWizard.cpp b/Engine/Plugins/Editor/MobileLauncherProfileWizard/Source/MobileLauncherProfileWizard/Private/AndroidProfileWizard.cpp
--- a/Engine/Plugins/Editor/MobileLauncherProfileWizard/Source/MobileLauncherProfileWizard/Private/AndroidProfileWizard.cpp
+++ b/Engine/Plugins/Editor/MobileLauncherProfileWizard/Source/MobileLauncherProfileWizard/Private/AndroidProfileWizard.cpp
@@ -27,6 +27,8 @@ namespace AndroidProfileConstants
 	const FString DLCName(TEXT("DLC1.0"));
 }
 
+PRAGMA_DISABLE_DEPRECATION_WARNINGS
+
 static void SetupAndroidAppProfile(ILauncherProfileRef& AppProfile, const FProfileParameters& Params, const FString& ProjectPath)
 {
 	AppProfile->SetProjectSpecified(true);
@@ -175,6 +177,8 @@ static void CreateAndroidProfiles(const FProfileParameters& Params, FString Proj
 	ProfileManager->SaveJSONProfile(DLCProfile);
 }
 
+PRAGMA_ENABLE_DEPRECATION_WARNINGS
+
 void FAndroidProfileWizard::HandleCreateLauncherProfile(const ILauncherProfileManagerRef& ProfileManager)
 {
 	const FVector2D WindowSize = FVector2D(940, 540);
diff --git a/Engine/Plugins/Editor/MobileLauncherProfileWizard/Source/MobileLauncherProfileWizard/Private/IOSProfileWizard.cpp b/Engine/Plugins/Editor/MobileLauncherProfileWizard/Source/MobileLauncherProfileWizard/Private/IOSProfileWizard.cpp
--- a/Engine/Plugins/Editor/MobileLauncherProfileWizard/Source/MobileLauncherProfileWizard/Private/IOSProfileWizard.cpp
+++ b/Engine/Plugins/Editor/MobileLauncherProfileWizard/Source/MobileLauncherProfileWizard/Private/IOSProfileWizard.cpp
@@ -27,6 +27,8 @@ namespace IOSProfileConstants
 	const FString DLCName(TEXT("DLC1.0"));
 }
 
+PRAGMA_DISABLE_DEPRECATION_WARNINGS
+
 static void SetupIOSAppProfile(ILauncherProfileRef& AppProfile, const FProfileParameters& Params, const FString& ProjectPath)
 {
 	AppProfile->SetProjectSpecified(true);
@@ -171,6 +173,8 @@ static void CreateIOSProfiles(const FProfileParameters& Params, FString ProjectP
 	ProfileManager->SaveJSONProfile(DLCProfile);
 }
 
+PRAGMA_ENABLE_DEPRECATION_WARNINGS
+
 void FIOSProfileWizard::HandleCreateLauncherProfile(const ILauncherProfileManagerRef& ProfileManager)
 {
 	const FVector2D WindowSize = FVector2D(940, 540);
diff --git a/Engine/Source/Developer/LauncherServices/Private/Launcher/LauncherWorker.cpp b/Engine/Source/Developer/LauncherServices/Private/Launcher/LauncherWorker.cpp
--- a/Engine/Source/Developer/LauncherServices/Private/Launcher/LauncherWorker.cpp
+++ b/Engine/Source/Developer/LauncherServices/Private/Launcher/LauncherWorker.cpp
@@ -231,14 +231,14 @@ int32 FLauncherWorker::GetTasks( TArray<ILauncherTaskPtr>& OutTasks ) const
 }
 
 
-void FLauncherWorker::OnTaskStarted(const FString& TaskName)
+void FLauncherWorker::OnTaskStarted(const FString& TaskName, ILauncherProfileUATCommandPtr UnifiedUATCommand)
 {
 	StageStartTime = FPlatformTime::Seconds();
 	StageStarted.Broadcast(TaskName);
 	// look for disconnected devices only after displaying "Running on..."
 	if(TaskName.Contains(TEXT("Run Task")))
 	{
-		EnableDeviceDiscoveryListener();
+		EnableDeviceDiscoveryListener(UnifiedUATCommand);
 	}
 }
 
@@ -248,12 +248,12 @@ void FLauncherWorker::OnTaskCompleted(const FString& TaskName)
 	StageCompleted.Broadcast(TaskName, FPlatformTime::Seconds() - StageStartTime);
 }
 
-static void AddDeviceToLaunchCommand(const FString& DeviceId, TSharedPtr<ITargetDeviceProxy> DeviceProxy, const ILauncherProfileRef& InProfile, FString& DeviceNames, FString& RoleCommands, bool& bVsyncAdded)
+static void AddDeviceToLaunchCommand(const FString& DeviceId, TSharedPtr<ITargetDeviceProxy> DeviceProxy, const ILauncherProfileBuildCookRunRef& InBuildCookRun, const ILauncherProfileRef& InProfile, FString& DeviceNames, FString& RoleCommands, bool& bVsyncAdded)
 {
 	// add the platform
 	DeviceNames += TEXT("+\"") + DeviceId + TEXT("\"");
 	TArray<ILauncherProfileLaunchRolePtr> Roles;
-	if (InProfile->GetLaunchRolesFor(DeviceId, Roles) > 0)
+	if (InBuildCookRun->GetLaunchRolesFor(DeviceId, Roles) > 0)
 	{
 		for (int32 RoleIndex = 0; RoleIndex < Roles.Num(); RoleIndex++)
 		{
@@ -564,7 +564,7 @@ FString FLauncherWorker::CreateUATCommandForBuildCookRun( const ILauncherProfile
 				TSharedPtr<ITargetDeviceProxy> DeviceProxy = DeviceProxyManager->FindProxyDeviceForTargetDevice(DeviceId);
 				if (DeviceProxy.IsValid())
 				{
-					AddDeviceToLaunchCommand(DeviceId, DeviceProxy, InProfile, DeviceNames, RoleCommands, bVsyncAdded);
+					AddDeviceToLaunchCommand(DeviceId, DeviceProxy, InBuildCookRun, InProfile, DeviceNames, RoleCommands, bVsyncAdded);
 
 					// also add the credentials, if necessary
 					FString DeviceUser = DeviceProxy->GetDeviceUser();
@@ -1224,22 +1224,49 @@ FString FLauncherWorker::CreateUATCommandForAutomatedTest( const ILauncherProfil
 	{
 		return FString();
 	}
-
+	
 	ILauncherProfileManagerRef LauncherProfileManager = FModuleManager::LoadModuleChecked<ILauncherServicesModule>("LauncherServices").GetProfileManager();
 
 
+	// @todo: could potentially add BuildConfiguation, LocalExecutable, Command Line, DeviceIDs and Platforms to each automated test so they can be run without a BCR... but maybe an overkill?
+	// ...for now, generate values based on unified values from the BuildCookRuns
+	EBuildConfiguration BuildConfiguration = EBuildConfiguration::Development;
+	bool bUsingLocalExecutables = false;
+	FString LaunchCommandLineParameters;
+	TArray<FString> DeviceIDs;
+	TArray<FString> SelectedPlatforms = InPlatforms;
+	for (ILauncherProfileBuildCookRunRef BuildCookRun : InProfile->GetBuildCookRunCommands())
+	{
+		if (BuildCookRun->IsEnabled())
+		{
+			BuildConfiguration = FMath::Max( BuildCookRun->GetBuildConfiguration(), BuildConfiguration );
+			LaunchCommandLineParameters += TEXT(" ");
+			LaunchCommandLineParameters += BuildCookRun->GetAdditionalCommandLineParameters();
+			bUsingLocalExecutables |= BuildCookRun->ShouldBuild();
+			if (BuildCookRun->GetDeployedDeviceGroup().IsValid())
+			{
+				for (const FString& DeviceID : BuildCookRun->GetDeployedDeviceGroup()->GetDeviceIDs())
+				{
+					DeviceIDs.AddUnique(DeviceID);
+				}
+			}
+		}
+	}
+	LaunchCommandLineParameters.TrimStartInline();
+
+
 	// basic command parameters
 	FString AutomatedTestCommand = FString::Printf( TEXT("%s -utf8output -nop4 "), *InAutomatedTest->GetUATCommand() );
 	AutomatedTestCommand += FString::Printf(TEXT("-project=\"%s\" "), *FPaths::ConvertRelativePathToFull(InProfile->GetProjectPath()));
-	AutomatedTestCommand += FString::Printf(TEXT("-configuration=%s "), LexToString(InProfile->GetBuildConfiguration()) );
+	AutomatedTestCommand += FString::Printf(TEXT("-configuration=%s "), LexToString(BuildConfiguration) );
 	AutomatedTestCommand += FString::Printf(TEXT("-tests=%s "), *InAutomatedTest->GetTests());
 
 	// collect platforms
 	FString ServerPlatforms = TEXT("");
 	FString Platforms = TEXT("");
-	for (int32 PlatformIndex = 0; PlatformIndex < InPlatforms.Num(); ++PlatformIndex)
+	for (int32 PlatformIndex = 0; PlatformIndex < SelectedPlatforms.Num(); ++PlatformIndex)
 	{
-		const PlatformInfo::FTargetPlatformInfo* PlatformInfo = PlatformInfo::FindPlatformInfo(FName(*InPlatforms[PlatformIndex]));
+		const PlatformInfo::FTargetPlatformInfo* PlatformInfo = PlatformInfo::FindPlatformInfo(FName(*SelectedPlatforms[PlatformIndex]));
 		if (ensure(PlatformInfo))
 		{
 			FString& PlatformString = (PlatformInfo->PlatformType == EBuildTargetType::Server) ? ServerPlatforms : Platforms;
@@ -1267,7 +1294,7 @@ FString FLauncherWorker::CreateUATCommandForAutomatedTest( const ILauncherProfil
 	}
 
 	// use the local executables we just made if necessary
-	if (InProfile->ShouldBuild())
+	if (bUsingLocalExecutables)
 	{
 		AutomatedTestCommand += TEXT("-dev ");
 	}
@@ -1287,11 +1314,10 @@ FString FLauncherWorker::CreateUATCommandForAutomatedTest( const ILauncherProfil
 	};
 
 	// specify the devices to use
-	ILauncherDeviceGroupPtr DeviceGroup = InProfile->GetDeployedDeviceGroup();
-	if (DeviceGroup.IsValid() && DeviceGroup->GetNumDevices() > 0)
+	if (DeviceIDs.Num() > 0)
 	{
 		FString DeviceNames;
-		for (const FString& DeviceId : DeviceGroup->GetDeviceIDs())
+		for (const FString& DeviceId : DeviceIDs)
 		{
 			FString FinalDeviceId = DeviceId.TrimStartAndEnd().Replace(TEXT("@"), TEXT(":"));
 			if (!FinalDeviceId.IsEmpty())
@@ -1309,11 +1335,11 @@ FString FLauncherWorker::CreateUATCommandForAutomatedTest( const ILauncherProfil
 	// it exists in the automated test equivalent which is -Args. 
 	FString AdditionalCommandLine;
 
-	// Append the profile command line parameters
-	FString ProfileAdditionalCommandLine = InProfile->GetAdditionalCommandLineParameters();
-	if (!ProfileAdditionalCommandLine.IsEmpty())
+	// Append the tests command line parameters
+	FString LaunchAdditionalCommandLine = LaunchCommandLineParameters;
+	if (!LaunchAdditionalCommandLine.IsEmpty())
 	{
-		AdditionalCommandLine = FString::Printf(TEXT("-Args=\"%s\""), *ProfileAdditionalCommandLine.Replace(TEXT("\""), TEXT("\\\"")));
+		AdditionalCommandLine = FString::Printf(TEXT("-Args=\"%s\""), *LaunchAdditionalCommandLine.Replace(TEXT("\""), TEXT("\\\"")));
 	}
 
 	// Append the tests additional command line parameters
@@ -1422,10 +1448,13 @@ FString FLauncherWorker::MakeBuildCookRunParamsForProjectCustomBuild(const ILaun
 {
 	// this function can only support one build cook run command (aka "classic launcher profiles")
 	TArray<ILauncherProfileBuildCookRunRef> AllBuildCookRunCmds = InProfile->GetBuildCookRunCommands();
-	if (AllBuildCookRunCmds.Num() == 0)
+	if (AllBuildCookRunCmds.Num() != 1)
 	{
 		return FString();
 	}
+	ILauncherProfileBuildCookRunRef BuildCookRun = AllBuildCookRunCmds[0];
+	
+
 
 	// get the basic UAT command line
 	FLauncherWorker TempWorker(NoInit);
@@ -1451,17 +1480,17 @@ FString FLauncherWorker::MakeBuildCookRunParamsForProjectCustomBuild(const ILaun
 	}
 
 	// optional device
-	if ( InProfile->GetDeploymentMode() == ELauncherProfileDeploymentModes::CopyToDevice || 
-		 InProfile->GetLaunchMode() != ELauncherProfileLaunchModes::DoNotLaunch ||
-		 InProfile->GetDeployedDeviceGroup() != nullptr )
+	if ( BuildCookRun->GetDeploymentMode() == ELauncherProfileDeploymentModes::CopyToDevice || 
+		 BuildCookRun->GetLaunchMode() != ELauncherProfileLaunchModes::DoNotLaunch ||
+		 BuildCookRun->GetDeployedDeviceGroup() != nullptr )
 	{
 		BuildCookRunParams += TEXT(" -device={DeviceId}");
 	}
 
 	// configuration
-	if (InProfile->GetBuildConfiguration() != EBuildConfiguration::Unknown)
+	if (BuildCookRun->GetBuildConfiguration() != EBuildConfiguration::Unknown)
 	{
-		BuildCookRunParams += FString::Printf( TEXT(" -configuration=%s"), LexToString(InProfile->GetBuildConfiguration()) );
+		BuildCookRunParams += FString::Printf( TEXT(" -configuration=%s"), LexToString(BuildCookRun->GetBuildConfiguration()) );
 	}
 
 	return BuildCookRunParams;
@@ -1652,9 +1681,9 @@ void FLauncherWorker::CreateAndExecuteTasks( const ILauncherProfileRef& InProfil
 
 	// create task chains
 	TSharedPtr<FLauncherTask> NextTask;
-	auto AddTask = [this, &NextTask](TSharedPtr<FLauncherTask> NewTask)
+	auto AddTask = [this, &NextTask](TSharedPtr<FLauncherTask> NewTask, ILauncherProfileUATCommandPtr UnifiedUATCommand = nullptr)
 	{
-		NewTask->OnStarted().AddRaw(this, &FLauncherWorker::OnTaskStarted);
+		NewTask->OnStarted().AddRaw(this, &FLauncherWorker::OnTaskStarted, UnifiedUATCommand);
 		NewTask->OnCompleted().AddRaw(this, &FLauncherWorker::OnTaskCompleted);
 
 		if (NextTask.IsValid())
@@ -1724,7 +1753,7 @@ void FLauncherWorker::CreateAndExecuteTasks( const ILauncherProfileRef& InProfil
 	}
 	else
 	{
-		// disable run in BuidCookRun commands if there is multiple BuildCookRun commands that want to run - in this case, they're gathered together and run in a separate unified run command at the end
+		// disable run in BuildCookRun commands if there is multiple BuildCookRun commands that want to run - in this case, they're gathered together and run in a separate unified run command at the end
 		UnifiedRunUATCommand = MakeTemporaryUnifiedRunCommand(InProfile);
 		if (UnifiedRunUATCommand.IsValid())
 		{
@@ -1773,7 +1802,7 @@ void FLauncherWorker::CreateAndExecuteTasks( const ILauncherProfileRef& InProfil
 
 
 	TSharedPtr<FLauncherTask> LaunchTask = MakeShareable(new FLauncherUATTask(UATCommandString, TEXT("Launch Task"), TEXT("Launching UAT..."), ReadPipe, WritePipe, ProcHandle, this, StartString));
-	AddTask(LaunchTask);
+	AddTask(LaunchTask, UnifiedRunUATCommand);
 	for (int32 Index = 0; Index < Commands.Num(); ++Index)
 	{
 		class FLauncherWaitTask : public FLauncherTask
@@ -1831,31 +1860,51 @@ void FLauncherWorker::CreateAndExecuteTasks( const ILauncherProfileRef& InProfil
 }
 
 /** Callback for when a device proxy has been removed from the device proxy manager. */
-void FLauncherWorker::HandleDeviceProxyManagerProxyRemoved(const TSharedRef<ITargetDeviceProxy>& RemovedProxy)
+void FLauncherWorker::HandleDeviceProxyManagerProxyRemoved(const TSharedRef<ITargetDeviceProxy>& RemovedProxy, ILauncherProfileUATCommandPtr UnifiedUATCommand)
 {
 	FString TargetDeviceId = RemovedProxy->GetTargetDeviceId(NAME_None);
-	
-	// determine deployment devices
-	ILauncherDeviceGroupPtr DeviceGroup = Profile->GetDeployedDeviceGroup();
 
-	if (DeviceGroup.IsValid() && DeviceGroup->GetNumDevices() > 0)
+	TArray<ILauncherProfileUATCommandRef> UATCommands = Profile->GetUATCommands();
+	if (UnifiedUATCommand.IsValid())
 	{
-		// remove disconnected device from the list
-		DeviceGroup->RemoveDevice(TargetDeviceId);
-		if (DeviceGroup->GetNumDevices() == 0)
+		UATCommands.Add(UnifiedUATCommand.ToSharedRef());
+	}
+
+	for (const ILauncherProfileUATCommandRef& UATCommand : UATCommands)
+	{
+		if (!UATCommand->IsEnabled())
 		{
-			// this was the last device running, stop working
-			Stop();
+			continue;
+		}
+
+		if (ILauncherProfileBuildCookRunPtr BuildCookRun = UATCommand->AsBuildCookRun())
+		{
+			if (BuildCookRun->GetLaunchMode() != ELauncherProfileLaunchModes::DoNotLaunch && (!UnifiedUATCommand.IsValid() || UnifiedUATCommand == UATCommand))
+			{
+				ILauncherDeviceGroupPtr DeviceGroup = BuildCookRun->GetDeployedDeviceGroup();
+
+				if (DeviceGroup.IsValid() && DeviceGroup->GetNumDevices() > 0)
+				{
+					// remove disconnected device from the list
+					DeviceGroup->RemoveDevice(TargetDeviceId);
+					if (DeviceGroup->GetNumDevices() == 0)
+					{
+						// this was the last device running, stop working
+						Stop();
+						break;
+					}
+				}
+			}
 		}
 	}
 }
 
 // start looking for disconnected devices
-void FLauncherWorker::EnableDeviceDiscoveryListener()
+void FLauncherWorker::EnableDeviceDiscoveryListener(ILauncherProfileUATCommandPtr UnifiedUATCommand)
 {
 	if (!OnProxyRemovedDelegateHandle.IsValid())
 	{
-		OnProxyRemovedDelegateHandle = DeviceProxyManager->OnProxyRemoved().AddRaw(this, &FLauncherWorker::HandleDeviceProxyManagerProxyRemoved);
+		OnProxyRemovedDelegateHandle = DeviceProxyManager->OnProxyRemoved().AddRaw(this, &FLauncherWorker::HandleDeviceProxyManagerProxyRemoved, UnifiedUATCommand);
 	}
 }
 
@@ -1873,15 +1922,23 @@ void FLauncherWorker::DisableDeviceDiscoveryListener()
 //Cancel the currently running application on all devices
 bool FLauncherWorker::TerminateLaunchedProcess()
 {
-	// determine deployment devices
-	ILauncherDeviceGroupPtr DeviceGroup = Profile->GetDeployedDeviceGroup();
-	FName Variant = NAME_None;
-
-	// device list
-	FString DeviceNamesParam = TEXT("");
-	if (DeviceGroup.IsValid())
+	for (const ILauncherProfileUATCommandRef& UATCommand : Profile->GetUATCommands())
 	{
-		const TArray<FString>& Devices = DeviceGroup->GetDeviceIDs();
+		TArray<FString> Devices;
+		if (const ILauncherProfileBuildCookRunPtr BuildCookRun = UATCommand->AsBuildCookRun())
+		{
+			// determine deployment devices
+			ILauncherDeviceGroupPtr DeviceGroup = BuildCookRun->GetDeployedDeviceGroup();
+			FName Variant = NAME_None;
+
+			// device list
+			FString DeviceNamesParam = TEXT("");
+			if (DeviceGroup.IsValid())
+			{
+				Devices = DeviceGroup->GetDeviceIDs();
+			}
+		}
+
 		// cancel the app onr each device
 		for (int32 DeviceIndex = 0; DeviceIndex < Devices.Num(); ++DeviceIndex)
 		{
diff --git a/Engine/Source/Developer/LauncherServices/Private/Launcher/LauncherWorker.h b/Engine/Source/Developer/LauncherServices/Private/Launcher/LauncherWorker.h
--- a/Engine/Source/Developer/LauncherServices/Private/Launcher/LauncherWorker.h
+++ b/Engine/Source/Developer/LauncherServices/Private/Launcher/LauncherWorker.h
@@ -88,7 +88,7 @@ public:
 	}
 
 	/** Callback for when a device proxy has been removed from the device proxy manager. */
-	void HandleDeviceProxyManagerProxyRemoved(const TSharedRef<ITargetDeviceProxy>& RemovedProxy);
+	void HandleDeviceProxyManagerProxyRemoved(const TSharedRef<ITargetDeviceProxy>& RemovedProxy, ILauncherProfileUATCommandPtr UnifiedUATCommand);
 
 	/**
 	 * Set the app id running on the device
@@ -113,7 +113,7 @@ protected:
 	 */
 	void CreateAndExecuteTasks( const ILauncherProfileRef& InProfile );
 
-	void OnTaskStarted(const FString& TaskName);
+	void OnTaskStarted(const FString& TaskName, ILauncherProfileUATCommandPtr UnifiedUATCommand);
 	void OnTaskCompleted(const FString& TaskName);
 
 	struct FCommandAdjustments
@@ -131,7 +131,7 @@ protected:
 	TArray<FString> FinalizePlatforms( const TArray<FString>& CookPlatforms, const ILauncherDeviceGroupPtr& DeviceGroup, bool bPreferCookedPlatforms ) const;
 
 	//start listening to the device discovery routine
-	void EnableDeviceDiscoveryListener();
+	void EnableDeviceDiscoveryListener(ILauncherProfileUATCommandPtr UnifiedUATCommand);
 
 	//stop listening to the device discovery routine
 	void DisableDeviceDiscoveryListener();
diff --git a/Engine/Source/Developer/LauncherServices/Public/ILauncherProfile.h b/Engine/Source/Developer/LauncherServices/Public/ILauncherProfile.h
--- a/Engine/Source/Developer/LauncherServices/Public/ILauncherProfile.h
+++ b/Engine/Source/Developer/LauncherServices/Public/ILauncherProfile.h
@@ -315,6 +315,7 @@ public:
 	 * @param PlatformName The name of the platform to deploy.
 	 * @return true if the platform is deployable, false otherwise.
 	 */
+	UE_DEPRECATED(5.8, "use GetFirstBuildCookRun version")
 	virtual bool IsDeployablePlatform( const FString& PlatformName ) = 0;
 
 	/**
@@ -392,6 +393,7 @@ public:
 	 * @return Build configuration name.
 	 * @see SetBuildConfigurationName
 	 */
+	UE_DEPRECATED(5.8, "use GetFirstBuildCookRun version")
 	virtual EBuildConfiguration GetBuildConfiguration( ) const = 0;
 
 	/**
@@ -399,6 +401,7 @@ public:
 	 *
 	 * @return true if the profile specifies a build target.
 	 */
+	UE_DEPRECATED(5.8, "use GetFirstBuildCookRun version")
 	virtual bool HasBuildTargetSpecified() const = 0;
 
 	/**
@@ -407,7 +410,7 @@ public:
 	 * @Return Target name to build/run (it would match a .Target.cs file)
 	 * @see SetBuildTarget
 	 */
-	UE_DEPRECATED(5.8, "Use GetBuildTargets instead")
+	UE_DEPRECATED(5.8, "use GetFirstBuildCookRun version and GetBuildTargets")
 	virtual FString GetBuildTarget() const = 0;
 
 	/**
@@ -416,6 +419,7 @@ public:
 	 * @return Cook configuration name.
 	 * @see SetCookConfigurationName
 	 */
+	UE_DEPRECATED(5.8, "use GetFirstBuildCookRun version")
 	virtual EBuildConfiguration GetCookConfiguration( ) const = 0;
 
 	/**
@@ -423,6 +427,7 @@ public:
 	 *
 	 * @return Cook mode.
 	 */
+	UE_DEPRECATED(5.8, "use GetFirstBuildCookRun version")
 	virtual ELauncherProfileCookModes::Type GetCookMode( ) const = 0;
 
 	/**
@@ -430,9 +435,11 @@ public:
 	 *
 	 * @return Cook options string.
 	 */
+	UE_DEPRECATED(5.8, "use GetFirstBuildCookRun version")
 	virtual const FString& GetCookOptions( ) const = 0;
 
 
+	UE_DEPRECATED(5.8, "use GetFirstBuildCookRun version")
 	virtual const bool GetSkipCookingEditorContent() const = 0; 
 
 	/**
@@ -441,6 +448,7 @@ public:
 	 * 
 	 * @param InSkipCookingEditorContent
 	 */
+	UE_DEPRECATED(5.8, "use GetFirstBuildCookRun version")
 	virtual void SetSkipCookingEditorContent(const bool InSkipCookingEditorContent) = 0;
 
 
@@ -450,6 +458,7 @@ public:
 	 * @return Collection of culture names.
 	 * @see AddCookedCulture, ClearCookedCultures, RemoveCookedCulture
 	 */
+	UE_DEPRECATED(5.8, "use GetFirstBuildCookRun version")
 	virtual const TArray<FString>& GetCookedCultures( ) const = 0;
 
 	/**
@@ -459,6 +468,7 @@ public:
 	 *
 	 * @see AddCookedMap, ClearCookedMaps, RemoveCookedMap
 	 */
+	UE_DEPRECATED(5.8, "use GetFirstBuildCookRun version")
 	virtual const TArray<FString>& GetCookedMaps( ) const = 0;
 
 	/**
@@ -468,6 +478,7 @@ public:
 	 *
 	 * @see AddCookedPlatform, ClearCookedPlatforms, RemoveCookedPlatform
 	 */
+	UE_DEPRECATED(5.8, "use GetFirstBuildCookRun version")
 	virtual const TArray<FString>& GetCookedPlatforms( ) const = 0;
 
 	/**
@@ -475,6 +486,7 @@ public:
 	 *
 	 * @return A reference to the default launch role.
 	 */
+	UE_DEPRECATED(5.8, "use GetFirstBuildCookRun version")
 	virtual const ILauncherProfileLaunchRoleRef& GetDefaultLaunchRole( ) const = 0;
 
 	/**
@@ -483,12 +495,14 @@ public:
 	 * @return The device group, or NULL if none was configured.
 	 * @see SetDeployedDeviceGroup
 	 */
+	UE_DEPRECATED(5.8, "use GetFirstBuildCookRun version")
 	virtual ILauncherDeviceGroupPtr GetDeployedDeviceGroup( ) = 0;
 
 	/**
 	* Gets the default platforms to deploy if no specific devices were selected.
 	*	
 	*/
+	UE_DEPRECATED(5.8, "use GetFirstBuildCookRun version")
 	virtual const FName GetDefaultDeployPlatform() const = 0;
 
 	/**
@@ -497,6 +511,7 @@ public:
 	 * @return The deployment mode.
 	 * @see SetDeploymentMode
 	 */
+	UE_DEPRECATED(5.8, "use GetFirstBuildCookRun version")
 	virtual ELauncherProfileDeploymentModes::Type GetDeploymentMode( ) const = 0;
 
     /**
@@ -505,6 +520,7 @@ public:
      * @return The close mode.
      * @see SetForceClose
      */
+	UE_DEPRECATED(5.8, "use GetFirstBuildCookRun version")
     virtual bool GetForceClose() const = 0;
     
 	/**
@@ -513,6 +529,7 @@ public:
 	 * @return The launch mode.
 	 * @see SetLaunchMode
 	 */
+	UE_DEPRECATED(5.8, "use GetFirstBuildCookRun version")
 	virtual ELauncherProfileLaunchModes::Type GetLaunchMode( ) const = 0;
 
 	/**
@@ -521,6 +538,7 @@ public:
 	 * @return A read-only collection of launch roles.
 	 * @see CreateLaunchRole, RemoveLaunchRole
 	 */
+	UE_DEPRECATED(5.8, "use GetFirstBuildCookRun version")
 	virtual const TArray<ILauncherProfileLaunchRolePtr>& GetLaunchRoles( ) const = 0;
 
 	/**
@@ -529,6 +547,7 @@ public:
 	 * @param DeviceId The identifier of the device.
 	 * @param OutRoles Will hold the assigned roles, if any.
 	 */
+	UE_DEPRECATED(5.8, "use GetFirstBuildCookRun version")
 	virtual const int32 GetLaunchRolesFor( const FString& DeviceId, TArray<ILauncherProfileLaunchRolePtr>& OutRoles ) = 0;
 
 	/**
@@ -537,6 +556,7 @@ public:
 	 * @return The packaging mode.
 	 * @see SetPackagingMode
 	 */
+	UE_DEPRECATED(5.8, "use GetFirstBuildCookRun version")
 	virtual ELauncherProfilePackagingModes::Type GetPackagingMode( ) const = 0;
 
 	/**
@@ -545,6 +565,7 @@ public:
 	 * @return The packaging directory.
 	 * @see SetPackageDirectory
 	 */
+	UE_DEPRECATED(5.8, "use GetFirstBuildCookRun version")
 	virtual FString GetPackageDirectory( ) const = 0;
 
 	/**
@@ -552,6 +573,7 @@ public:
 	 *
 	 * @see SetArchive
 	 */
+	UE_DEPRECATED(5.8, "use GetFirstBuildCookRun version")
 	virtual bool IsArchiving( ) const = 0;
 
 	/**
@@ -560,6 +582,7 @@ public:
 	 * @return The archive directory.
 	 * @see SetArchiveDirectory
 	 */
+	UE_DEPRECATED(5.8, "use GetFirstBuildCookRun version")
 	virtual FString GetArchiveDirectory( ) const = 0;
 
 	/**
@@ -595,6 +618,7 @@ public:
 	 * @return The additional command line parameters
 	 * @see SetAdditionalCommandLineParameters
 	 */
+	UE_DEPRECATED(5.8, "use GetFirstBuildCookRun version")
 	virtual FString GetAdditionalCommandLineParameters() const = 0;
 
 	/**
@@ -604,6 +628,7 @@ public:
 	 * @return The additional command line parameters
 	 * @see SetAdditionalTargetCommandLineParameters
 	 */
+	UE_DEPRECATED(5.8, "use GetFirstBuildCookRun version")
 	virtual FString GetAdditionalTargetCommandLineParameters( EBuildTargetType BuildTargetType ) const = 0;
 
     /**
@@ -612,12 +637,14 @@ public:
      * @return The timeout time.
      * @see SetTimeout
      */
+	UE_DEPRECATED(5.8, "use GetFirstBuildCookRun version")
     virtual uint32 GetTimeout() const = 0;
 	
 	/**
 	 * Are we going to generate a patch (Source content patch needs to be specified)
 	 * @Seealso GetPatchSourceContentPath	 * @Seealso SetPatchSourceContentPath
 	 */
+	UE_DEPRECATED(5.8, "use GetFirstBuildCookRun version")
 	virtual bool IsGeneratingPatch() const = 0;
 	
 	/**
@@ -637,6 +664,7 @@ public:
 	 * @return The current build mode.
 	 * @see SetBuildMode
 	 */
+	UE_DEPRECATED(5.8, "use GetFirstBuildCookRun version")
 	virtual ELauncherProfileBuildModes::Type GetBuildMode() const = 0;
 
 	/**
@@ -645,6 +673,7 @@ public:
 	 * @return The current build mode.
 	 * @see SetBuildMode
 	 */
+	UE_DEPRECATED(5.8, "use GetFirstBuildCookRun version")
 	virtual bool ShouldBuild() = 0;
 
 	/**
@@ -653,6 +682,7 @@ public:
 	 * @return true if fast iterate is enabled, false otherwise.
 	 * @see SetFastIterate
 	 */
+	UE_DEPRECATED(5.8, "use GetFirstBuildCookRun version")
 	virtual bool IsFastIterate() const = 0;
 	
 	/**
@@ -669,6 +699,7 @@ public:
 	 * @return true if cooking incrementally, false otherwise.
 	 * @see SetIncrementalCooking
 	 */
+	UE_DEPRECATED(5.8, "use GetFirstBuildCookRun version")
 	virtual bool IsCookingIncrementally( ) const = 0;
 
 	/**
@@ -677,10 +708,10 @@ public:
 	 * @return incremental cook mode
 	 * @see SetIncrementalCookMode
 	 */ 
-	UE_DEPRECATED(5.8, "Use IsCookingIncrementally again instead")
+	UE_DEPRECATED(5.8, "use GetFirstBuildCookRun version and IsCookingIncrementally again instead")
 	virtual ELauncherProfileIncrementalCookMode::Type GetIncrementalCookMode() const = 0;
 
-
+	UE_DEPRECATED(5.8, "use GetFirstBuildCookRun version")
 	virtual bool IsIterateSharedCookedBuild() const =0;
 
 	/**
@@ -689,6 +720,7 @@ public:
 	 * @return true if compression is enabled
 	 * @see SetCompressed
 	 */
+	UE_DEPRECATED(5.8, "use GetFirstBuildCookRun version")
 	virtual bool IsCompressed( ) const = 0;
 
 	/**
@@ -696,6 +728,7 @@ public:
 	 * 
 	 * @return true if encrypting ini files is enabled
 	 */
+	UE_DEPRECATED(5.8, "use GetFirstBuildCookRun version")
 	virtual bool IsEncryptingIniFiles() const = 0;
 
 	/**
@@ -703,6 +736,7 @@ public:
 	 *
 	 * @return true if encrypting ini files is enabled
 	 */
+	UE_DEPRECATED(5.8, "use GetFirstBuildCookRun version")
 	virtual bool IsForDistribution() const = 0;
 
 	/**
@@ -711,6 +745,7 @@ public:
 	 * @return true if cooking unversioned, false otherwise.
 	 * @see SetUnversionedCooking
 	 */
+	UE_DEPRECATED(5.8, "use GetFirstBuildCookRun version")
 	virtual bool IsCookingUnversioned( ) const = 0;
 
 	/**
@@ -719,6 +754,7 @@ public:
 	 * @return true if deploying incrementally, false otherwise.
 	 * @see SetIncrementalDeploying
 	 */
+	UE_DEPRECATED(5.8, "use GetFirstBuildCookRun version")
 	virtual bool IsDeployingIncrementally( ) const = 0;
 
 	/**
@@ -727,6 +763,7 @@ public:
 	 * @return true if the file server should be hidden, false otherwise.
 	 * @see SetHideFileServer
 	 */
+	UE_DEPRECATED(5.8, "use GetFirstBuildCookRun version")
 	virtual bool IsFileServerHidden( ) const = 0;
 
 	/**
@@ -735,6 +772,7 @@ public:
 	 * @return true if the file server is streaming, false otherwise.
 	 * @see SetStreamingFileServer
 	 */
+	UE_DEPRECATED(5.8, "use GetFirstBuildCookRun version")
 	virtual bool IsFileServerStreaming( ) const = 0;
 
 	/**
@@ -743,6 +781,7 @@ public:
 	 * @return true if UnrealPak is used, false otherwise.
 	 * @see SetPackageWithUnrealPak
 	 */
+	UE_DEPRECATED(5.8, "use GetFirstBuildCookRun version")
 	virtual bool IsPackingWithUnrealPak( ) const = 0;
 
 	/**
@@ -751,6 +790,7 @@ public:
 	 * @return true if prerequisites are to be included, false otherwise.
 	 * @see SetIncludePrerequisites
 	 */
+	UE_DEPRECATED(5.8, "use GetFirstBuildCookRun version")
 	virtual bool IsIncludingPrerequisites() const = 0;
 
 	/**
@@ -758,6 +798,7 @@ public:
 	 *
 	 * @return true if Chunks will be generated, false otherwise.
 	 */
+	UE_DEPRECATED(5.8, "use GetFirstBuildCookRun version")
 	virtual bool IsGeneratingChunks() const = 0;
 	
 	/**
@@ -765,6 +806,7 @@ public:
 	 *
 	 * @return true if data will be generated, false otherwise.
 	 */
+	UE_DEPRECATED(5.8, "use GetFirstBuildCookRun version")
 	virtual bool IsGenerateHttpChunkData() const = 0;
 	
 	/**
@@ -772,6 +814,7 @@ public:
 	 *
 	 * @return true if data will be generated, false otherwise.
 	 */
+	UE_DEPRECATED(5.8, "use GetFirstBuildCookRun version")
 	virtual FString GetHttpChunkDataDirectory() const = 0;
 	
 	/**
@@ -779,6 +822,7 @@ public:
 	 *
 	 * @return true if data will be generated, false otherwise.
 	 */
+	UE_DEPRECATED(5.8, "use GetFirstBuildCookRun version")
 	virtual FString GetHttpChunkDataReleaseName() const = 0;
 
 	/**
@@ -811,6 +855,7 @@ public:
 	 *
 	 * @see ClearCookedCultures, GetCookedCultures, RemoveCookedCulture
 	 */
+	UE_DEPRECATED(5.8, "use GetFirstBuildCookRun version")
 	virtual void AddCookedCulture( const FString& CultureName ) = 0;
 
 	/**
@@ -820,6 +865,7 @@ public:
 	 *
 	 * @see ClearCookedMaps, GetCookedMaps, RemoveCookedMap
 	 */
+	UE_DEPRECATED(5.8, "use GetFirstBuildCookRun version")
 	virtual void AddCookedMap( const FString& MapName ) = 0;
 
 	/**
@@ -829,6 +875,7 @@ public:
 	 *
 	 * @see ClearCookedPlatforms, GetCookedPlatforms, RemoveCookedPlatform
 	 */
+	UE_DEPRECATED(5.8, "use GetFirstBuildCookRun version")
 	virtual void AddCookedPlatform( const FString& PlatformName ) = 0;
 
 	/**
@@ -838,6 +885,7 @@ public:
 	*
 	* @param PlatformName The name of the platform to add.		
 	*/
+	UE_DEPRECATED(5.8, "use GetFirstBuildCookRun version")
 	virtual void SetDefaultDeployPlatform(const FName PlatformName) = 0;
 
 	/**
@@ -845,6 +893,7 @@ public:
 	 *
 	 * @see AddCookedCulture, GetCookedCulture, RemoveCookedCulture
 	 */
+	UE_DEPRECATED(5.8, "use GetFirstBuildCookRun version")
 	virtual void ClearCookedCultures( ) = 0;
 
 	/**
@@ -852,6 +901,7 @@ public:
 	 *
 	 * @see AddCookedMap, GetCookedMap, RemoveCookedMap
 	 */
+	UE_DEPRECATED(5.8, "use GetFirstBuildCookRun version")
 	virtual void ClearCookedMaps( ) = 0;
 
 	/**
@@ -859,6 +909,7 @@ public:
 	 *
 	 * @see AddCookedPlatform, GetCookedPlatforms, RemoveCookedPlatform
 	 */
+	UE_DEPRECATED(5.8, "use GetFirstBuildCookRun version")
 	virtual void ClearCookedPlatforms( ) = 0;
 
 	/**
@@ -867,6 +918,7 @@ public:
 	 * @return The created role.
 	 * @see GetLaunchRoles, RemoveLaunchRole
 	 */
+	UE_DEPRECATED(5.8, "use GetFirstBuildCookRun version")
 	virtual ILauncherProfileLaunchRolePtr CreateLaunchRole( ) = 0;
 
 	/**
@@ -875,6 +927,7 @@ public:
 	 * @param CultureName The name of the culture to remove.
 	 * @see AddCookedCulture, ClearCookedCultures, GetCookedCultures
 	 */
+	UE_DEPRECATED(5.8, "use GetFirstBuildCookRun version")
 	virtual void RemoveCookedCulture( const FString& CultureName ) = 0;
 
 	/**
@@ -883,6 +936,7 @@ public:
 	 * @param MapName The name of the map to remove.
 	 * @see AddCookedMap, ClearCookedMaps, GetCookedMaps
 	 */
+	UE_DEPRECATED(5.8, "use GetFirstBuildCookRun version")
 	virtual void RemoveCookedMap( const FString& MapName ) = 0;
 
 	/**
@@ -891,6 +945,7 @@ public:
 	 * @param PlatformName The name of the platform to remove.
 	 * @see AddBuildPlatform, ClearCookedPlatforms, GetBuildPlatforms
 	 */
+	UE_DEPRECATED(5.8, "use GetFirstBuildCookRun version")
 	virtual void RemoveCookedPlatform( const FString& PlatformName ) = 0;
 
 	/**
@@ -899,6 +954,7 @@ public:
 	 * @param Role The role to remove.
 	 * @see CreateLaunchRole, GetLaunchRoles
 	 */
+	UE_DEPRECATED(5.8, "use GetFirstBuildCookRun version")
 	virtual void RemoveLaunchRole( const ILauncherProfileLaunchRoleRef& Role ) = 0;
 
 	/**
@@ -907,6 +963,7 @@ public:
 	 * @param Mode Whether the game should be built.
 	 * @see GetBuildMode
 	 */
+	UE_DEPRECATED(5.8, "use GetFirstBuildCookRun version")
 	virtual void SetBuildMode(ELauncherProfileBuildModes::Type Mode) = 0;
 
 	/**
@@ -915,6 +972,7 @@ public:
 	 * @param Enable fast iterate.
 	 * @see IsFastIterate
 	 */
+	UE_DEPRECATED(5.8, "use GetFirstBuildCookRun version")
 	virtual void SetFastIterate(bool Enable) = 0;
 
 	/**
@@ -931,6 +989,7 @@ public:
 	 * @param ConfigurationName The build configuration name to set.
 	 * @see GetBuildConfigurationName
 	 */
+	UE_DEPRECATED(5.8, "use GetFirstBuildCookRun version")
 	virtual void SetBuildConfiguration( EBuildConfiguration Configuration ) = 0;
 
 	/**
@@ -938,6 +997,7 @@ public:
 	 * 
 	 * @param Specified Whether a build target is specified.
 	 */
+	UE_DEPRECATED(5.8, "use GetFirstBuildCookRun version")
 	virtual void SetBuildTargetSpecified(bool Specified) = 0;
 
 	/** Notifies the profile that the fallback build target changed. */
@@ -949,7 +1009,7 @@ public:
 	 * @param TargetName The target name to set (it would match a .Target.cs file)
 	 * @see GetBuildTarget
 	 */
-	UE_DEPRECATED(5.8, "Use AddBuildTarget instead")
+	UE_DEPRECATED(5.8, "use GetFirstBuildCookRun version and AddBuildTarget instead")
 	virtual void SetBuildTarget( const FString& TargetName ) = 0;
 
 	/**
@@ -958,6 +1018,7 @@ public:
 	 * @param Configuration The cooker's build configuration to set.
 	 * @see GetCookConfiguration
 	 */
+	UE_DEPRECATED(5.8, "use GetFirstBuildCookRun version")
 	virtual void SetCookConfiguration( EBuildConfiguration Configuration ) = 0;
 
 	/**
@@ -966,6 +1027,7 @@ public:
 	 * @param Mode The cook mode.
 	 * @see GetCookMode
 	 */
+	UE_DEPRECATED(5.8, "use GetFirstBuildCookRun version")
 	virtual void SetCookMode( ELauncherProfileCookModes::Type Mode ) = 0;
 
 	/**
@@ -974,6 +1036,7 @@ public:
 	 * @param Options The cook options.
 	 * @see GetCookOptions
 	 */
+	UE_DEPRECATED(5.8, "use GetFirstBuildCookRun version")
 	virtual void SetCookOptions(const FString& Options) = 0;
 
 	/**
@@ -982,6 +1045,7 @@ public:
 	 * @param UseUnrealPak Whether UnrealPak should be used.
 	 * @see IsPackingWithUnrealPak
 	 */
+	UE_DEPRECATED(5.8, "use GetFirstBuildCookRun version")
 	virtual void SetDeployWithUnrealPak( bool UseUnrealPak ) = 0;
 
 	/**
@@ -989,24 +1053,28 @@ public:
 	 *
 	 * @param true if Chunks should be generated, false otherwise.
 	 */
+	UE_DEPRECATED(5.8, "use GetFirstBuildCookRun version")
 	virtual void SetGenerateChunks(bool bGenerateChunks) = 0;
 	/**
 	 * Set whether packaging will use chunk data to generate http chunk install data.
 	 *
 	 * @param true if data should be generated, false otherwise.
 	 */
+	UE_DEPRECATED(5.8, "use GetFirstBuildCookRun version")
 	virtual void SetGenerateHttpChunkData(bool bGenerateHttpChunkData) = 0;
 	/**
 	 * Set where generated http chunk install data will be stored.
 	 *
 	 * @return the directory path to use.
 	 */	
+	UE_DEPRECATED(5.8, "use GetFirstBuildCookRun version")
 	virtual void SetHttpChunkDataDirectory(const FString& InHttpChunkDataDirectory ) = 0;
 	/**
 	 * Set what name to tag the generated http chunk install data with.
 	 *
 	 * @param the name to use.
 	 */	
+	UE_DEPRECATED(5.8, "use GetFirstBuildCookRun version")
 	virtual void SetHttpChunkDataReleaseName(const FString& InHttpChunkDataReleaseName ) = 0;
 
 	/**
@@ -1015,6 +1083,7 @@ public:
 	 * @param DeviceGroup The device group, or NULL to reset this setting.
 	 * @see GetDeployedDeviceGroup
 	 */
+	UE_DEPRECATED(5.8, "use GetFirstBuildCookRun version")
 	virtual void SetDeployedDeviceGroup( const ILauncherDeviceGroupPtr& DeviceGroup ) = 0;
 
 	/**
@@ -1023,34 +1092,47 @@ public:
 	 * @param Mode The deployment mode to set.
 	 * @see GetDeploymentMode
 	 */
+	UE_DEPRECATED(5.8, "use GetFirstBuildCookRun version")
 	virtual void SetDeploymentMode( ELauncherProfileDeploymentModes::Type Mode ) = 0;
 
 	/**
 	 * Creating a release version of the cooked content 
 	 */
+	UE_DEPRECATED(5.8, "use GetFirstBuildCookRun version")
 	virtual bool IsCreatingReleaseVersion() const = 0;
 
+	UE_DEPRECATED(5.8, "use GetFirstBuildCookRun version")
 	virtual void SetCreateReleaseVersion(bool InCreateReleaseVersion) = 0;
 
+	UE_DEPRECATED(5.8, "use GetFirstBuildCookRun version")
 	virtual FString GetCreateReleaseVersionName() const = 0;
 
+	UE_DEPRECATED(5.8, "use GetFirstBuildCookRun version")
 	virtual void SetCreateReleaseVersionName(const FString& InCreateReleaseVersionName) = 0;
 
+	UE_DEPRECATED(5.8, "use GetFirstBuildCookRun version")
 	virtual FString GetBasedOnReleaseVersionName() const = 0;
 
+	UE_DEPRECATED(5.8, "use GetFirstBuildCookRun version")
 	virtual void SetBasedOnReleaseVersionName(const FString& InBasedOnReleaseVersion) = 0;
 
+	UE_DEPRECATED(5.8, "use GetFirstBuildCookRun version")
 	virtual FString GetOriginalReleaseVersionName() const = 0;
 
+	UE_DEPRECATED(5.8, "use GetFirstBuildCookRun version")
 	virtual void SetOriginalReleaseVersionName(const FString& InOriginalReleaseVersion) = 0;
 
 	/**
 	* Provides a database of compressed iostore chunks to reuse during the
 	* staging process. See IoStoreUtilities.cpp ReferenceContainerGlobalFileName.
 	*/
+	UE_DEPRECATED(5.8, "use GetFirstBuildCookRun version")
 	virtual FString GetReferenceContainerGlobalFileName() const = 0;
+	UE_DEPRECATED(5.8, "use GetFirstBuildCookRun version")
 	virtual void SetReferenceContainerGlobalFileName(const FString& InReferenceContainerGlobalFileName) = 0;
+	UE_DEPRECATED(5.8, "use GetFirstBuildCookRun version")
 	virtual FString GetReferenceContainerCryptoKeysFileName() const = 0;
+	UE_DEPRECATED(5.8, "use GetFirstBuildCookRun version")
 	virtual void SetReferenceContainerCryptoKeysFileName(const FString& InReferenceContainerCryptoKeysFileName) = 0;
 
 	/**
@@ -1059,18 +1141,26 @@ public:
 	 * @param InShouldGeneratePatch enable generating patch
 	 * @seealso IsGeneratingPatch
 	 */
+	UE_DEPRECATED(5.8, "use GetFirstBuildCookRun version")
 	virtual void SetGeneratePatch( bool InShouldGeneratePatch ) = 0;
+	UE_DEPRECATED(5.8, "use GetFirstBuildCookRun version")
 	virtual void SetAddPatchLevel( bool InAddPatchLevel) = 0;
 	UE_DEPRECATED(5.6, "ShouldStageBaseReleasePaks is no longer used")
 	virtual void SetStageBaseReleasePaks(bool InStageBaseReleasePaks) = 0;
 
+	UE_DEPRECATED(5.8, "use GetFirstBuildCookRun version")
 	virtual bool IsCreatingDLC() const = 0;
+	UE_DEPRECATED(5.8, "use GetFirstBuildCookRun version")
 	virtual void SetCreateDLC(bool InBuildDLC) = 0;
 
+	UE_DEPRECATED(5.8, "use GetFirstBuildCookRun version")
 	virtual FString GetDLCName() const = 0;
+	UE_DEPRECATED(5.8, "use GetFirstBuildCookRun version")
 	virtual void SetDLCName(const FString& InDLCName) = 0;
 
+	UE_DEPRECATED(5.8, "use GetFirstBuildCookRun version")
 	virtual bool IsDLCIncludingEngineContent() const = 0;
+	UE_DEPRECATED(5.8, "use GetFirstBuildCookRun version")
 	virtual void SetDLCIncludeEngineContent(bool InDLCIncludeEngineContent) = 0;
 
 
@@ -1080,6 +1170,7 @@ public:
      * @param Close the close mode to set.
      * @see GetForceClose
      */
+	UE_DEPRECATED(5.8, "use GetFirstBuildCookRun version")
     virtual void SetForceClose( bool Close ) = 0;
     
 	/**
@@ -1088,6 +1179,7 @@ public:
 	 * @param Hide Whether to hide the window.
 	 * @see GetHideFileServerWindow
 	 */
+	UE_DEPRECATED(5.8, "use GetFirstBuildCookRun version")
 	virtual void SetHideFileServerWindow( bool Hide ) = 0;
 
 	/**
@@ -1096,6 +1188,7 @@ public:
 	 * @param Incremental Whether cooking should be incremental.
 	 * @see IsCookingIncrementally
 	 */
+	UE_DEPRECATED(5.8, "use GetFirstBuildCookRun version")
 	virtual void SetIncrementalCooking( bool Incremental ) = 0;
 
 	/**
@@ -1104,10 +1197,10 @@ public:
 	 * @param Mode the incremental cooking mode to use
 	 * @see GetIncrementalCookMode
 	 */
-	UE_DEPRECATED(5.8, "Use SetIncrementalCooking again instead")
+	UE_DEPRECATED(5.8, "use GetFirstBuildCookRun version with SetIncrementalCooking again instead")
 	virtual void SetIncrementalCookMode( ELauncherProfileIncrementalCookMode::Type Mode ) = 0;
 
-
+	UE_DEPRECATED(5.8, "use GetFirstBuildCookRun version")
 	virtual void SetIterateSharedCookedBuild( bool IterateSharedCookedBuild ) = 0;
 
 	/**
@@ -1116,6 +1209,7 @@ public:
 	 * @param Enable compression
 	 * @see IsCompressed
 	 */
+	UE_DEPRECATED(5.8, "use GetFirstBuildCookRun version")
 	virtual void SetCompressed( bool Enable ) = 0;
 
 
@@ -1125,6 +1219,7 @@ public:
 	 * @param Enable encrypt ini files
 	 * @see IsEncryptIniFiles
 	 */
+	UE_DEPRECATED(5.8, "use GetFirstBuildCookRun version")
 	virtual void SetEncryptingIniFiles(bool Enabled) = 0;
 
 	/**
@@ -1133,6 +1228,7 @@ public:
 	* @param enable for distribution
 	* @see IsForDistribution
 	*/
+	UE_DEPRECATED(5.8, "use GetFirstBuildCookRun version")
 	virtual void SetForDistribution(bool Enabled) = 0;
 	
 
@@ -1142,6 +1238,7 @@ public:
 	 * @param Incremental Whether deploying should be incremental.
 	 * @see IsDeployingIncrementally
 	 */
+	UE_DEPRECATED(5.8, "use GetFirstBuildCookRun version")
 	virtual void SetIncrementalDeploying( bool Incremental ) = 0;
 
 	/**
@@ -1150,6 +1247,7 @@ public:
 	 * @param Mode The launch mode to set.
 	 * @see GetLaunchMode
 	 */
+	UE_DEPRECATED(5.8, "use GetFirstBuildCookRun version")
 	virtual void SetLaunchMode( ELauncherProfileLaunchModes::Type Mode ) = 0;
 
 	/**
@@ -1158,6 +1256,7 @@ public:
 	 * @param Mode The packaging mode to set.
 	 * @see GetPackagingMode
 	 */
+	UE_DEPRECATED(5.8, "use GetFirstBuildCookRun version")
 	virtual void SetPackagingMode( ELauncherProfilePackagingModes::Type Mode ) = 0;
 
 	/**
@@ -1166,6 +1265,7 @@ public:
 	 * @param Dir The packaging directory to set.
 	 * @see GetPackageDirectory
 	 */
+	UE_DEPRECATED(5.8, "use GetFirstBuildCookRun version")
 	virtual void SetPackageDirectory( const FString& Dir ) = 0;
 
 	/**
@@ -1173,6 +1273,7 @@ public:
 	 *
 	 * @see GetArchiveMode
 	 */
+	UE_DEPRECATED(5.8, "use GetFirstBuildCookRun version")
 	virtual void SetArchive( bool bArchive ) = 0;
 
 	/**
@@ -1181,6 +1282,7 @@ public:
 	 * @param Dir The archive directory to set.
 	 * @see GetArchiveDirectory
 	 */
+	UE_DEPRECATED(5.8, "use GetFirstBuildCookRun version")
 	virtual void SetArchiveDirectory( const FString& Dir ) = 0;
 
 	/**
@@ -1208,6 +1310,7 @@ public:
 	 * @param	Params	The additional command line parameters to use
 	 * @see GetAdditionalCommandLineParameters
 	 */
+	UE_DEPRECATED(5.8, "use GetFirstBuildCookRun version")
 	virtual void SetAdditionalCommandLineParameters(const FString& Params) = 0;
 
 	/**
@@ -1217,6 +1320,7 @@ public:
 	 * @param	BuildTargetType	The type of build target to modify
 	 * @see GetAdditionalTargetCommandLineParameters
 	 */
+	UE_DEPRECATED(5.8, "use GetFirstBuildCookRun version")
 	virtual void SetAdditionalTargetCommandLineParameters(const FString& Params, EBuildTargetType BuildTargetType ) = 0;
 
 
@@ -1226,6 +1330,7 @@ public:
 	 * @param Streaming Whether a streaming server should be used.
 	 * @see GetStreamingFileServer
 	 */
+	UE_DEPRECATED(5.8, "use GetFirstBuildCookRun version")
 	virtual void SetStreamingFileServer( bool Streaming ) = 0;
 
 	/**
@@ -1234,6 +1339,7 @@ public:
 	 * @param Value Whether prerequisites should be used.
 	 * @see IsIncludingPrerequisites
 	 */
+	UE_DEPRECATED(5.8, "use GetFirstBuildCookRun version")
 	virtual void SetIncludePrerequisites(bool InValue) = 0;
 
     /**
@@ -1242,6 +1348,7 @@ public:
      * @param InTime Amount of time to wait before timing out.
      * @see GetTimeout
      */
+	UE_DEPRECATED(5.8, "use GetFirstBuildCookRun version")
     virtual void SetTimeout(uint32 InTime) = 0;
     
 	/**
@@ -1250,6 +1357,7 @@ public:
 	 * @param Unversioned Whether cooking is unversioned.
 	 * @see IsCookingUnversioned
 	 */
+	UE_DEPRECATED(5.8, "use GetFirstBuildCookRun version")
 	virtual void SetUnversionedCooking( bool Unversioned ) = 0;
 
 	/**
@@ -1284,6 +1392,7 @@ public:
 	 * Sets whether to use I/O store for optimized loading.
 	 * @param bUseIoStore Whether to use I/O store
 	 */
+	UE_DEPRECATED(5.8, "use GetFirstBuildCookRun version")
 	virtual void SetUseIoStore(bool bUseIoStore) = 0;
 
 	/**
@@ -1291,12 +1400,14 @@ public:
 	 *
 	 * @return true if using I/O store
 	 */
+	UE_DEPRECATED(5.8, "use GetFirstBuildCookRun version")
 	virtual bool IsUsingIoStore() const = 0;
 
 	/**
 	 * Sets whether to use the Zen storage server. Note, this property is ignored if the project itself is configured to use Zen store (via UProjectPackagingSettings::bUseZenStore)
 	 * @param bUseZenStore Whether to use the Zen storage server
 	 */
+	UE_DEPRECATED(5.8, "use GetFirstBuildCookRun version")
 	virtual void SetUseZenStore(bool bUseZenStore) = 0;
 
 	/**
@@ -1304,6 +1415,7 @@ public:
 	 *
 	 * @return true if using Zen storage server
 	 */
+	UE_DEPRECATED(5.8, "use GetFirstBuildCookRun version")
 	virtual bool IsUsingZenStore() const = 0;
 
 	/**
@@ -1318,6 +1430,7 @@ public:
 	 * 
 	 * @param Import Whether to import a Zen snapshot
 	 */
+	UE_DEPRECATED(5.8, "use GetFirstBuildCookRun version")
 	virtual void SetImportingZenSnapshot( bool Import ) = 0;
 
 	/**
@@ -1325,6 +1438,7 @@ public:
 	 * 
 	 * @return true if importing a Zen snapshot, false otherwise
 	 */
+	UE_DEPRECATED(5.8, "use GetFirstBuildCookRun version")
 	virtual bool IsImportingZenSnapshot() const = 0;
 
 	/**
@@ -1332,6 +1446,7 @@ public:
 	 * 
 	 * @param UseZenPakStreaming Whether use Zen pak streaming.
 	 */
+	UE_DEPRECATED(5.8, "use GetFirstBuildCookRun version")
 	virtual void SetUseZenPakStreaming( bool UseZenPakStreaming ) = 0;
 
 	/**
@@ -1339,6 +1454,7 @@ public:
 	 * 
 	 * @return true if using Zen pak streaming, false otherwise
 	 */
+	UE_DEPRECATED(5.8, "use GetFirstBuildCookRun version")
 	virtual bool IsUsingZenPakStreaming() const = 0;
 
 	/**
@@ -1347,6 +1463,7 @@ public:
 	 * 
 	 * @param Path Path to the build to use for Zen pak streaming
 	 */
+	UE_DEPRECATED(5.8, "use GetFirstBuildCookRun version")
 	virtual void SetZenPakStreamingPath( const FString& Path ) = 0;
 
 	/**
@@ -1354,18 +1471,21 @@ public:
 	 * 
 	 * @return current Zen pak streaming path, if any
 	 */
+	UE_DEPRECATED(5.8, "use GetFirstBuildCookRun version")
 	virtual FString GetZenPakStreamingPath() const = 0;
 
 	/**
 	 * Sets whether to make a binary config file during packaging
 	 * @param bMakeBinaryConfig Whether to make a binary config file during staging
 	 */
+	UE_DEPRECATED(5.8, "use GetFirstBuildCookRun version")
 	virtual void SetMakeBinaryConfig(bool bMakeBinaryConfig) = 0;
 	/**
 	 * Make binary config file during staging or not.
 	 *
 	 * @return true to make binary config file
 	 */
+	UE_DEPRECATED(5.8, "use GetFirstBuildCookRun version")
 	virtual bool MakeBinaryConfig() const = 0;
 
 	/**
@@ -1381,41 +1501,49 @@ public:
 	/**
 	 * Sets whether or not the Device is a Simulator
 	 */
+	UE_DEPRECATED(5.8, "use GetFirstBuildCookRun version")
 	virtual void SetDeviceIsASimulator(bool bInIsDeviceASimualtor) = 0;
 
 	/**
 	 * Whether or not the Device is a Simulator
 	 */
+	UE_DEPRECATED(5.8, "use GetFirstBuildCookRun version")
 	virtual bool IsDeviceASimulator() const = 0;
 
 	/** 
 	 * Sets the client architecture(s) to build (x64, arm64 etc). Leave empty to use the default
 	 */
+	UE_DEPRECATED(5.8, "use GetFirstBuildCookRun version")
 	virtual void SetClientArchitectures( const TArray<FString>& InArchitectures ) = 0;
 
 	/**
 	 * Returns the client architecture(s) that will be build. empty means the default will be used
 	 */
+	UE_DEPRECATED(5.8, "use GetFirstBuildCookRun version")
 	virtual const TArray<FString>& GetClientArchitectures() const = 0;
 	
 	/** 
 	 * Sets the server architecture(s) to build (x64, arm64 etc). Leave empty to use the default
 	 */
+	UE_DEPRECATED(5.8, "use GetFirstBuildCookRun version")
 	virtual void SetServerArchitectures( const TArray<FString>& InArchitectures ) = 0;
 
 	/**
 	 * Returns the server architecture(s) that will be build. empty means the default will be used
 	 */
+	UE_DEPRECATED(5.8, "use GetFirstBuildCookRun version")
 	virtual const TArray<FString>& GetServerArchitectures() const = 0;
 
 	/** 
 	 * Sets the editor architecture(s) to build (x64, arm64 etc). Leave empty to use the default
 	 */
+	UE_DEPRECATED(5.8, "use GetFirstBuildCookRun version")
 	virtual void SetEditorArchitectures( const TArray<FString>& InArchitectures ) = 0;
 
 	/**
 	 * Returns the editor architecture(s) that will be build. empty means the default will be used
 	 */
+	UE_DEPRECATED(5.8, "use GetFirstBuildCookRun version")
 	virtual const TArray<FString>& GetEditorArchitectures() const = 0;
 
 
@@ -1529,24 +1657,28 @@ public:
 	 * Gets the names of the build targets
 	 * @see AddBuildTarget, RemoveBuildTarget, ClearBuildTargets
 	 */
+	UE_DEPRECATED(5.8, "use GetFirstBuildCookRun version")
 	virtual TArray<FString> GetBuildTargets() const = 0;
 
 	/**
 	 * Add a build target to the list
 	 * @see GetBuildTargets, RemoveBuildTarget, ClearBuildTargets
 	 */
+	UE_DEPRECATED(5.8, "use GetFirstBuildCookRun version")
 	virtual void AddBuildTarget( const FString& InBuildTarget ) = 0;
 
 	/**
 	 * Remove a build target from the list
 	 * @see GetBuildTargets, AddBuildTarget, ClearBuildTargets
 	 */
+	UE_DEPRECATED(5.8, "use GetFirstBuildCookRun version")
 	virtual void RemoveBuildTarget( const FString& InBuildTarget ) = 0;
 
 	/**
 	 * Clears the list of build targets
 	 * @see GetBuildTargets, AddBuildTarget, RemoveBuildTarget
 	 */
+	UE_DEPRECATED(5.8, "use GetFirstBuildCookRun version")
 	virtual void ClearBuildTargets() = 0;
 
 
diff --git a/Engine/Source/Developer/LegacyProjectLauncher/Private/Widgets/Archive/SProjectLauncherArchivePage.cpp b/Engine/Source/Developer/LegacyProjectLauncher/Private/Widgets/Archive/SProjectLauncherArchivePage.cpp
--- a/Engine/Source/Developer/LegacyProjectLauncher/Private/Widgets/Archive/SProjectLauncherArchivePage.cpp
+++ b/Engine/Source/Developer/LegacyProjectLauncher/Private/Widgets/Archive/SProjectLauncherArchivePage.cpp
@@ -116,6 +116,8 @@ void SProjectLauncherArchivePage::Construct(const FArguments& InArgs, const TSha
 /* SProjectLauncherArchivePage callbacks
  *****************************************************************************/
 
+PRAGMA_DISABLE_DEPRECATION_WARNINGS
+
 void SProjectLauncherArchivePage::HandleArchiveCheckedStateChanged(ECheckBoxState CheckState)
 {
 	ILauncherProfilePtr SelectedProfile = Model->GetSelectedProfile();
@@ -233,5 +235,6 @@ void SProjectLauncherArchivePage::OnDirectoryTextCommitted(const FText& InText,
 	}
 }
 
+PRAGMA_ENABLE_DEPRECATION_WARNINGS
 
 #undef LOCTEXT_NAMESPACE
diff --git a/Engine/Source/Developer/LegacyProjectLauncher/Private/Widgets/Build/SProjectLauncherBuildPage.cpp b/Engine/Source/Developer/LegacyProjectLauncher/Private/Widgets/Build/SProjectLauncherBuildPage.cpp
--- a/Engine/Source/Developer/LegacyProjectLauncher/Private/Widgets/Build/SProjectLauncherBuildPage.cpp
+++ b/Engine/Source/Developer/LegacyProjectLauncher/Private/Widgets/Build/SProjectLauncherBuildPage.cpp
@@ -234,6 +234,8 @@ bool SProjectLauncherBuildPage::GenerateDSYMForProject(const FString& ProjectNam
 /* SProjectLauncherBuildPage callbacks
  *****************************************************************************/
 
+PRAGMA_DISABLE_DEPRECATION_WARNINGS
+
 FText SProjectLauncherBuildPage::HandleBuildModeComboButtonContentText() const
 {
 	ILauncherProfilePtr SelectedProfile = Model->GetSelectedProfile();
@@ -434,4 +436,6 @@ FText SProjectLauncherBuildPage::GetCommandLineText() const
 	return FText::FromString(CommandLine);
 }
 
+PRAGMA_ENABLE_DEPRECATION_WARNINGS
+
 #undef LOCTEXT_NAMESPACE
diff --git a/Engine/Source/Developer/LegacyProjectLauncher/Private/Widgets/Cook/SProjectLauncherCookByTheBookSettings.cpp b/Engine/Source/Developer/LegacyProjectLauncher/Private/Widgets/Cook/SProjectLauncherCookByTheBookSettings.cpp
--- a/Engine/Source/Developer/LegacyProjectLauncher/Private/Widgets/Cook/SProjectLauncherCookByTheBookSettings.cpp
+++ b/Engine/Source/Developer/LegacyProjectLauncher/Private/Widgets/Cook/SProjectLauncherCookByTheBookSettings.cpp
@@ -1161,6 +1161,7 @@ void SProjectLauncherCookByTheBookSettings::RefreshCultureList()
 	}
 }
 
+PRAGMA_DISABLE_DEPRECATION_WARNINGS
 
 void SProjectLauncherCookByTheBookSettings::RefreshMapList()
 {
@@ -2037,5 +2038,6 @@ ECheckBoxState SProjectLauncherCookByTheBookSettings::HandleDontIncludeEditorCon
 	return ECheckBoxState::Unchecked;
 }
 
+PRAGMA_ENABLE_DEPRECATION_WARNINGS
 
 #undef LOCTEXT_NAMESPACE
diff --git a/Engine/Source/Developer/LegacyProjectLauncher/Private/Widgets/Cook/SProjectLauncherCookOnTheFlySettings.cpp b/Engine/Source/Developer/LegacyProjectLauncher/Private/Widgets/Cook/SProjectLauncherCookOnTheFlySettings.cpp
--- a/Engine/Source/Developer/LegacyProjectLauncher/Private/Widgets/Cook/SProjectLauncherCookOnTheFlySettings.cpp
+++ b/Engine/Source/Developer/LegacyProjectLauncher/Private/Widgets/Cook/SProjectLauncherCookOnTheFlySettings.cpp
@@ -98,6 +98,8 @@ END_SLATE_FUNCTION_BUILD_OPTIMIZATION
 /* SProjectLauncherCookOnTheFlySettings callbacks
  *****************************************************************************/
 
+PRAGMA_DISABLE_DEPRECATION_WARNINGS
+
 void SProjectLauncherCookOnTheFlySettings::HandleIncrementalCheckBoxCheckStateChanged(ECheckBoxState NewState)
 {
 	ILauncherProfilePtr SelectedProfile = Model->GetSelectedProfile();
@@ -183,4 +185,6 @@ void SProjectLauncherCookOnTheFlySettings::HandleCookerOptionsCommitted(const FT
 	}
 }
 
+PRAGMA_ENABLE_DEPRECATION_WARNINGS
+
 #undef LOCTEXT_NAMESPACE
diff --git a/Engine/Source/Developer/LegacyProjectLauncher/Private/Widgets/Cook/SProjectLauncherCookPage.cpp b/Engine/Source/Developer/LegacyProjectLauncher/Private/Widgets/Cook/SProjectLauncherCookPage.cpp
--- a/Engine/Source/Developer/LegacyProjectLauncher/Private/Widgets/Cook/SProjectLauncherCookPage.cpp
+++ b/Engine/Source/Developer/LegacyProjectLauncher/Private/Widgets/Cook/SProjectLauncherCookPage.cpp
@@ -108,6 +108,8 @@ void SProjectLauncherCookPage::Construct(const FArguments& InArgs, const TShared
 /* SProjectLauncherCookPage callbacks
  *****************************************************************************/
 
+PRAGMA_DISABLE_DEPRECATION_WARNINGS
+
 EVisibility SProjectLauncherCookPage::HandleCookByTheBookSettingsVisibility() const
 {
 	ILauncherProfilePtr SelectedProfile = Model->GetSelectedProfile();
@@ -186,5 +188,6 @@ void SProjectLauncherCookPage::HandleProfileManagerProfileSelected(const ILaunch
 	// reload settings
 }
 
+PRAGMA_ENABLE_DEPRECATION_WARNINGS
 
 #undef LOCTEXT_NAMESPACE
diff --git a/Engine/Source/Developer/LegacyProjectLauncher/Private/Widgets/Cook/SProjectLauncherCookedPlatforms.h b/Engine/Source/Developer/LegacyProjectLauncher/Private/Widgets/Cook/SProjectLauncherCookedPlatforms.h
--- a/Engine/Source/Developer/LegacyProjectLauncher/Private/Widgets/Cook/SProjectLauncherCookedPlatforms.h
+++ b/Engine/Source/Developer/LegacyProjectLauncher/Private/Widgets/Cook/SProjectLauncherCookedPlatforms.h
@@ -68,6 +68,7 @@ private:
 
 		if (SelectedProfile.IsValid())
 		{
+PRAGMA_DISABLE_DEPRECATION_WARNINGS
 			if (AllPlatforms)
 			{
 				TArray<ITargetPlatform*> Platforms = GetTargetPlatformManager()->GetTargetPlatforms();
@@ -81,6 +82,7 @@ private:
 			{
 				SelectedProfile->ClearCookedPlatforms();
 			}
+PRAGMA_ENABLE_DEPRECATION_WARNINGS
 		}
 	}
 
diff --git a/Engine/Source/Developer/LegacyProjectLauncher/Private/Widgets/Cook/SProjectLauncherCultureListRow.h b/Engine/Source/Developer/LegacyProjectLauncher/Private/Widgets/Cook/SProjectLauncherCultureListRow.h
--- a/Engine/Source/Developer/LegacyProjectLauncher/Private/Widgets/Cook/SProjectLauncherCultureListRow.h
+++ b/Engine/Source/Developer/LegacyProjectLauncher/Private/Widgets/Cook/SProjectLauncherCultureListRow.h
@@ -76,6 +76,8 @@ public:
 
 private:
 
+PRAGMA_DISABLE_DEPRECATION_WARNINGS
+
 	// Callback for changing the checked state of the check box.
 	void HandleCheckBoxCheckStateChanged( ECheckBoxState NewState )
 	{
@@ -107,6 +109,8 @@ private:
 		return ECheckBoxState::Unchecked;
 	}
 
+PRAGMA_ENABLE_DEPRECATION_WARNINGS
+
 private:
 
 	// Holds the highlight string for the log message.
diff --git a/Engine/Source/Developer/LegacyProjectLauncher/Private/Widgets/Cook/SProjectLauncherMapListRow.h b/Engine/Source/Developer/LegacyProjectLauncher/Private/Widgets/Cook/SProjectLauncherMapListRow.h
--- a/Engine/Source/Developer/LegacyProjectLauncher/Private/Widgets/Cook/SProjectLauncherMapListRow.h
+++ b/Engine/Source/Developer/LegacyProjectLauncher/Private/Widgets/Cook/SProjectLauncherMapListRow.h
@@ -75,6 +75,8 @@ public:
 
 private:
 
+PRAGMA_DISABLE_DEPRECATION_WARNINGS
+
 	// Callback for changing the checked state of the check box.
 	void HandleCheckBoxCheckStateChanged( ECheckBoxState NewState )
 	{
@@ -106,6 +108,8 @@ private:
 		return ECheckBoxState::Unchecked;
 	}
 
+PRAGMA_ENABLE_DEPRECATION_WARNINGS
+
 private:
 
 	// Holds the highlight string for the log message.
diff --git a/Engine/Source/Developer/LegacyProjectLauncher/Private/Widgets/Cook/SProjectLauncherPlatformListRow.h b/Engine/Source/Developer/LegacyProjectLauncher/Private/Widgets/Cook/SProjectLauncherPlatformListRow.h
--- a/Engine/Source/Developer/LegacyProjectLauncher/Private/Widgets/Cook/SProjectLauncherPlatformListRow.h
+++ b/Engine/Source/Developer/LegacyProjectLauncher/Private/Widgets/Cook/SProjectLauncherPlatformListRow.h
@@ -75,6 +75,8 @@ public:
 
 private:
 
+PRAGMA_DISABLE_DEPRECATION_WARNINGS
+
 	// Callback for changing the checked state of the check box.
 	void HandleCheckBoxCheckStateChanged( ECheckBoxState NewState )
 	{
@@ -110,6 +112,8 @@ private:
 		return ECheckBoxState::Unchecked;
 	}
 
+PRAGMA_ENABLE_DEPRECATION_WARNINGS
+
 private:
 
 	// Holds the highlight string for the log message.
diff --git a/Engine/Source/Developer/LegacyProjectLauncher/Private/Widgets/Cook/SProjectLauncherSimpleCookPage.cpp b/Engine/Source/Developer/LegacyProjectLauncher/Private/Widgets/Cook/SProjectLauncherSimpleCookPage.cpp
--- a/Engine/Source/Developer/LegacyProjectLauncher/Private/Widgets/Cook/SProjectLauncherSimpleCookPage.cpp
+++ b/Engine/Source/Developer/LegacyProjectLauncher/Private/Widgets/Cook/SProjectLauncherSimpleCookPage.cpp
@@ -47,6 +47,8 @@ void SProjectLauncherSimpleCookPage::Construct(const FArguments& InArgs, const T
 /* SProjectLauncherSimpleCookPage callbacks
  *****************************************************************************/
 
+PRAGMA_DISABLE_DEPRECATION_WARNINGS
+
 EVisibility SProjectLauncherSimpleCookPage::HandleCookByTheBookSettingsVisibility() const
 {
 	ILauncherProfilePtr SelectedProfile = Model->GetSelectedProfile();
@@ -125,5 +127,6 @@ void SProjectLauncherSimpleCookPage::HandleProfileManagerProfileSelected(const I
 	// reload settings
 }
 
+PRAGMA_ENABLE_DEPRECATION_WARNINGS
 
 #undef LOCTEXT_NAMESPACE
diff --git a/Engine/Source/Developer/LegacyProjectLauncher/Private/Widgets/Deploy/SProjectLauncherDeployFileServerSettings.cpp b/Engine/Source/Developer/LegacyProjectLauncher/Private/Widgets/Deploy/SProjectLauncherDeployFileServerSettings.cpp
--- a/Engine/Source/Developer/LegacyProjectLauncher/Private/Widgets/Deploy/SProjectLauncherDeployFileServerSettings.cpp
+++ b/Engine/Source/Developer/LegacyProjectLauncher/Private/Widgets/Deploy/SProjectLauncherDeployFileServerSettings.cpp
@@ -93,6 +93,8 @@ void SProjectLauncherDeployFileServerSettings::Construct(const FArguments& InArg
 /* SProjectLauncherDeployFileServerSettings callbacks
  *****************************************************************************/
 
+PRAGMA_DISABLE_DEPRECATION_WARNINGS
+
 void SProjectLauncherDeployFileServerSettings::HandleHideWindowCheckBoxCheckStateChanged(ECheckBoxState NewState)
 {
 	ILauncherProfilePtr SelectedProfile = Model->GetSelectedProfile();
@@ -146,5 +148,6 @@ ECheckBoxState SProjectLauncherDeployFileServerSettings::HandleStreamingServerCh
 	return ECheckBoxState::Unchecked;
 }
 
+PRAGMA_ENABLE_DEPRECATION_WARNINGS
 
 #undef LOCTEXT_NAMESPACE
diff --git a/Engine/Source/Developer/LegacyProjectLauncher/Private/Widgets/Deploy/SProjectLauncherDeployPage.cpp b/Engine/Source/Developer/LegacyProjectLauncher/Private/Widgets/Deploy/SProjectLauncherDeployPage.cpp
--- a/Engine/Source/Developer/LegacyProjectLauncher/Private/Widgets/Deploy/SProjectLauncherDeployPage.cpp
+++ b/Engine/Source/Developer/LegacyProjectLauncher/Private/Widgets/Deploy/SProjectLauncherDeployPage.cpp
@@ -188,6 +188,8 @@ END_SLATE_FUNCTION_BUILD_OPTIMIZATION
 /* SProjectLauncherDevicesPage callbacks
  *****************************************************************************/
 
+PRAGMA_DISABLE_DEPRECATION_WARNINGS
+
 FText SProjectLauncherDeployPage::HandleDeploymentModeComboButtonContentText() const
 {
 	ILauncherProfilePtr SelectedProfile = Model->GetSelectedProfile();
@@ -297,5 +299,6 @@ EVisibility SProjectLauncherDeployPage::HandleValidationErrorIconVisibility(ELau
 	return EVisibility::Collapsed;
 }
 
+PRAGMA_ENABLE_DEPRECATION_WARNINGS
 
 #undef LOCTEXT_NAMESPACE
diff --git a/Engine/Source/Developer/LegacyProjectLauncher/Private/Widgets/Deploy/SProjectLauncherDeployRepositorySettings.cpp b/Engine/Source/Developer/LegacyProjectLauncher/Private/Widgets/Deploy/SProjectLauncherDeployRepositorySettings.cpp
--- a/Engine/Source/Developer/LegacyProjectLauncher/Private/Widgets/Deploy/SProjectLauncherDeployRepositorySettings.cpp
+++ b/Engine/Source/Developer/LegacyProjectLauncher/Private/Widgets/Deploy/SProjectLauncherDeployRepositorySettings.cpp
@@ -95,6 +95,8 @@ void SProjectLauncherDeployRepositorySettings::Construct(const FArguments& InArg
 /* SProjectLauncherDeployRepositorySettings callbacks
  *****************************************************************************/
 
+PRAGMA_DISABLE_DEPRECATION_WARNINGS
+
 FReply SProjectLauncherDeployRepositorySettings::HandleBrowseButtonClicked()
 {
 	IDesktopPlatform* DesktopPlatform = FDesktopPlatformModule::Get();
@@ -156,5 +158,6 @@ void SProjectLauncherDeployRepositorySettings::OnTextCommitted(const FText& InTe
 	}
 }
 
+PRAGMA_ENABLE_DEPRECATION_WARNINGS
 
 #undef LOCTEXT_NAMESPACE
diff --git a/Engine/Source/Developer/LegacyProjectLauncher/Private/Widgets/Deploy/SProjectLauncherDeployTargets.cpp b/Engine/Source/Developer/LegacyProjectLauncher/Private/Widgets/Deploy/SProjectLauncherDeployTargets.cpp
--- a/Engine/Source/Developer/LegacyProjectLauncher/Private/Widgets/Deploy/SProjectLauncherDeployTargets.cpp
+++ b/Engine/Source/Developer/LegacyProjectLauncher/Private/Widgets/Deploy/SProjectLauncherDeployTargets.cpp
@@ -160,6 +160,8 @@ void SProjectLauncherDeployTargets::RefreshDeviceProxyList()
 /* SSessionSProjectLauncherDeployTargets callbacks
  *****************************************************************************/
 
+PRAGMA_DISABLE_DEPRECATION_WARNINGS
+
 ILauncherDeviceGroupPtr SProjectLauncherDeployTargets::HandleDeviceListRowDeviceGroup() const
 {
 	ILauncherProfilePtr SelectedProfile = Model->GetSelectedProfile();
@@ -322,5 +324,6 @@ void SProjectLauncherDeployTargets::HandlePlatformComboBoxSelectionChanged(TShar
 	} 	
 }
 
+PRAGMA_ENABLE_DEPRECATION_WARNINGS
 
 #undef LOCTEXT_NAMESPACE
diff --git a/Engine/Source/Developer/LegacyProjectLauncher/Private/Widgets/Deploy/SProjectLauncherDeployToDeviceSettings.cpp b/Engine/Source/Developer/LegacyProjectLauncher/Private/Widgets/Deploy/SProjectLauncherDeployToDeviceSettings.cpp
--- a/Engine/Source/Developer/LegacyProjectLauncher/Private/Widgets/Deploy/SProjectLauncherDeployToDeviceSettings.cpp
+++ b/Engine/Source/Developer/LegacyProjectLauncher/Private/Widgets/Deploy/SProjectLauncherDeployToDeviceSettings.cpp
@@ -82,6 +82,8 @@ void SProjectLauncherDeployToDeviceSettings::Construct(const FArguments& InArgs,
 /* SProjectLauncherDeployToDeviceSettings callbacks
  *****************************************************************************/
 
+PRAGMA_DISABLE_DEPRECATION_WARNINGS
+
 void SProjectLauncherDeployToDeviceSettings::HandleIncrementalCheckBoxCheckStateChanged(ECheckBoxState NewState)
 {
 	ILauncherProfilePtr SelectedProfile = Model->GetSelectedProfile();
@@ -107,4 +109,7 @@ ECheckBoxState SProjectLauncherDeployToDeviceSettings::HandleIncrementalCheckBox
 
 	return ECheckBoxState::Unchecked;
 }
+
+PRAGMA_ENABLE_DEPRECATION_WARNINGS
+
 #undef LOCTEXT_NAMESPACE
diff --git a/Engine/Source/Developer/LegacyProjectLauncher/Private/Widgets/Deploy/SProjectLauncherSimpleDeviceListRow.h b/Engine/Source/Developer/LegacyProjectLauncher/Private/Widgets/Deploy/SProjectLauncherSimpleDeviceListRow.h
--- a/Engine/Source/Developer/LegacyProjectLauncher/Private/Widgets/Deploy/SProjectLauncherSimpleDeviceListRow.h
+++ b/Engine/Source/Developer/LegacyProjectLauncher/Private/Widgets/Deploy/SProjectLauncherSimpleDeviceListRow.h
@@ -87,9 +87,11 @@ public:
 
 		SimpleProfile = Model->GetProfileManager()->FindOrAddSimpleProfile(DeviceProxy->GetName());
 
+PRAGMA_DISABLE_DEPRECATION_WARNINGS
 		LaunchProfile = Model->GetProfileManager()->CreateUnsavedProfile(DeviceProxy->GetName());
 		LaunchProfile->SetProjectSpecified(false); // device launch profiles should always use the fallback project & build target
 		LaunchProfile->SetBuildTargetSpecified(false);
+PRAGMA_ENABLE_DEPRECATION_WARNINGS
 		UpdateProfile();
 
 		TSharedRef<SUniformGridPanel> NameGrid = SNew(SUniformGridPanel).SlotPadding(FMargin(0.0f, 1.0f));
@@ -365,6 +367,7 @@ private:
 		}
 	}
 
+PRAGMA_DISABLE_DEPRECATION_WARNINGS
 	void UpdateProfile()
 	{
 		check(DeviceProxy.IsValid());
@@ -412,6 +415,7 @@ private:
 			}
 		}
 	}
+PRAGMA_ENABLE_DEPRECATION_WARNINGS
 
 	ILauncherProfilePtr GetLaunchProfile() const
 	{
diff --git a/Engine/Source/Developer/LegacyProjectLauncher/Private/Widgets/Launch/SProjectLauncherLaunchPage.cpp b/Engine/Source/Developer/LegacyProjectLauncher/Private/Widgets/Launch/SProjectLauncherLaunchPage.cpp
--- a/Engine/Source/Developer/LegacyProjectLauncher/Private/Widgets/Launch/SProjectLauncherLaunchPage.cpp
+++ b/Engine/Source/Developer/LegacyProjectLauncher/Private/Widgets/Launch/SProjectLauncherLaunchPage.cpp
@@ -159,6 +159,7 @@ void SProjectLauncherLaunchPage::Construct(const FArguments& InArgs, const TShar
 }
 END_SLATE_FUNCTION_BUILD_OPTIMIZATION
 
+PRAGMA_DISABLE_DEPRECATION_WARNINGS
 
 void SProjectLauncherLaunchPage::Refresh()
 {
@@ -319,5 +320,6 @@ EVisibility SProjectLauncherLaunchPage::HandleValidationErrorIconVisibility(ELau
 	return EVisibility::Collapsed;
 }
 
+PRAGMA_ENABLE_DEPRECATION_WARNINGS
 
 #undef LOCTEXT_NAMESPACE
diff --git a/Engine/Source/Developer/LegacyProjectLauncher/Private/Widgets/Package/SProjectLauncherPackagePage.cpp b/Engine/Source/Developer/LegacyProjectLauncher/Private/Widgets/Package/SProjectLauncherPackagePage.cpp
--- a/Engine/Source/Developer/LegacyProjectLauncher/Private/Widgets/Package/SProjectLauncherPackagePage.cpp
+++ b/Engine/Source/Developer/LegacyProjectLauncher/Private/Widgets/Package/SProjectLauncherPackagePage.cpp
@@ -99,6 +99,8 @@ void SProjectLauncherPackagePage::Construct(const FArguments& InArgs, const TSha
 /* SProjectLauncherLaunchPage callbacks
  *****************************************************************************/
 
+PRAGMA_DISABLE_DEPRECATION_WARNINGS
+
 FText SProjectLauncherPackagePage::HandlePackagingModeComboButtonContentText() const
 {
 	ILauncherProfilePtr SelectedProfile = Model->GetSelectedProfile();
@@ -166,5 +168,6 @@ void SProjectLauncherPackagePage::HandleProfileManagerProfileSelected(const ILau
 
 }
 
+PRAGMA_ENABLE_DEPRECATION_WARNINGS
 
 #undef LOCTEXT_NAMESPACE
diff --git a/Engine/Source/Developer/LegacyProjectLauncher/Private/Widgets/Package/SProjectLauncherPackagingSettings.cpp b/Engine/Source/Developer/LegacyProjectLauncher/Private/Widgets/Package/SProjectLauncherPackagingSettings.cpp
--- a/Engine/Source/Developer/LegacyProjectLauncher/Private/Widgets/Package/SProjectLauncherPackagingSettings.cpp
+++ b/Engine/Source/Developer/LegacyProjectLauncher/Private/Widgets/Package/SProjectLauncherPackagingSettings.cpp
@@ -245,6 +245,7 @@ void SProjectLauncherPackagingSettings::Construct(const FArguments& InArgs, cons
 }
 END_SLATE_FUNCTION_BUILD_OPTIMIZATION
 
+PRAGMA_DISABLE_DEPRECATION_WARNINGS
 
 void SProjectLauncherPackagingSettings::UpdatePackagingModeWidgets()
 {
@@ -611,4 +612,6 @@ FReply SProjectLauncherPackagingSettings::HandleRefChunkCryptoBrowseButtonClicke
 	return FReply::Handled();
 }
 
+PRAGMA_ENABLE_DEPRECATION_WARNINGS
+
 #undef LOCTEXT_NAMESPACE
diff --git a/Engine/Source/Developer/LegacyProjectLauncher/Private/Widgets/Preview/SProjectLauncherPreviewPage.cpp b/Engine/Source/Developer/LegacyProjectLauncher/Private/Widgets/Preview/SProjectLauncherPreviewPage.cpp
--- a/Engine/Source/Developer/LegacyProjectLauncher/Private/Widgets/Preview/SProjectLauncherPreviewPage.cpp
+++ b/Engine/Source/Developer/LegacyProjectLauncher/Private/Widgets/Preview/SProjectLauncherPreviewPage.cpp
@@ -810,6 +810,8 @@ void SProjectLauncherPreviewPage::Tick( const FGeometry& AllottedGeometry, const
 /* SProjectLauncherDevicesSelector implementation
  *****************************************************************************/
 
+PRAGMA_DISABLE_DEPRECATION_WARNINGS
+
 void SProjectLauncherPreviewPage::RefreshDeviceProxyList( )
 {
 	DeviceProxyList.Reset();
@@ -1184,5 +1186,6 @@ EVisibility SProjectLauncherPreviewPage::HandleValidationErrorIconVisibility( EL
 	return EVisibility::Collapsed;
 }
 
+PRAGMA_ENABLE_DEPRECATION_WARNINGS
 
 #undef LOCTEXT_NAMESPACE
diff --git a/Engine/Source/Developer/LegacyProjectLauncher/Private/Widgets/Shared/SProjectLauncherBuildTargetSelector.cpp b/Engine/Source/Developer/LegacyProjectLauncher/Private/Widgets/Shared/SProjectLauncherBuildTargetSelector.cpp
--- a/Engine/Source/Developer/LegacyProjectLauncher/Private/Widgets/Shared/SProjectLauncherBuildTargetSelector.cpp
+++ b/Engine/Source/Developer/LegacyProjectLauncher/Private/Widgets/Shared/SProjectLauncherBuildTargetSelector.cpp
@@ -105,6 +105,8 @@ FText SProjectLauncherBuildTargetSelector::GetBuildTargetText( FString BuildTarg
 	return LOCTEXT("NoBuildTargetText", "Default");
 }
 
+PRAGMA_DISABLE_DEPRECATION_WARNINGS
+
 FText SProjectLauncherBuildTargetSelector::HandleBuildTargetComboButtonText() const
 {
 	FString BuildTarget;
@@ -256,4 +258,6 @@ ILauncherProfilePtr SProjectLauncherBuildTargetSelector::GetProfile() const
 	return nullptr;
 }
 
+PRAGMA_ENABLE_DEPRECATION_WARNINGS
+
 #undef LOCTEXT_NAMESPACE
diff --git a/Engine/Source/Editor/UnrealEd/Private/PlayLevelViaLauncher.cpp b/Engine/Source/Editor/UnrealEd/Private/PlayLevelViaLauncher.cpp
--- a/Engine/Source/Editor/UnrealEd/Private/PlayLevelViaLauncher.cpp
+++ b/Engine/Source/Editor/UnrealEd/Private/PlayLevelViaLauncher.cpp
@@ -213,17 +213,18 @@ void UEditorEngine::StartPlayUsingLauncherSession(FRequestPlaySessionParams& InR
 
 	// Setup launch profile, keep the setting here to a minimum.
 	ILauncherProfileRef LauncherProfile = LauncherServicesModule.CreateProfile(TEXT("Launch On Device"));
-	LauncherProfile->SetBuildMode(BuildMode);
-	LauncherProfile->SetBuildConfiguration(BuildConfiguration);
-	LauncherProfile->SetFastIterate(true);
+	ILauncherProfileBuildCookRunPtr BuildCookRun = LauncherProfile->GetFirstBuildCookRun();
+	BuildCookRun->SetBuildMode(BuildMode);
+	BuildCookRun->SetBuildConfiguration(BuildConfiguration);
+	BuildCookRun->SetFastIterate(true);
 	if (InRequestParams.EditorPlaySettings && !InRequestParams.EditorPlaySettings->AdditionalLaunchParameters.IsEmpty())
 	{
-		LauncherProfile->SetAdditionalCommandLineParameters(InRequestParams.EditorPlaySettings->AdditionalLaunchParameters);
+		BuildCookRun->SetAdditionalCommandLineParameters(InRequestParams.EditorPlaySettings->AdditionalLaunchParameters);
 	}
 
-	LauncherProfile->AddCookedPlatform(LaunchPlatformName);
+	BuildCookRun->AddCookedPlatform(LaunchPlatformName);
 
-	LauncherProfile->SetDeviceIsASimulator(InRequestParams.LauncherTargetDevice->bIsSimulator);
+	BuildCookRun->SetDeviceIsASimulator(InRequestParams.LauncherTargetDevice->bIsSimulator);
 
 	// select the quickest cook mode based on which in editor cook mode is enabled
 	const UCookerSettings& CookerSettings = *GetDefault<UCookerSettings>();
@@ -234,14 +235,14 @@ void UEditorEngine::StartPlayUsingLauncherSession(FRequestPlaySessionParams& InR
 	ELauncherProfileCookModes::Type CurrentLauncherCookMode = ELauncherProfileCookModes::ByTheBook;
 	if (!CookerSettings.bCookOnTheFlyForLaunchOn)
 	{
-		bInEditorCooking = Algo::AllOf(LauncherProfile->GetCookedPlatforms(),
+		bInEditorCooking = Algo::AllOf(BuildCookRun->GetCookedPlatforms(),
 			[this](const FString& PlatformName) { return CanCookByTheBookInEditor(PlatformName); });
 		CurrentLauncherCookMode = bInEditorCooking ? ELauncherProfileCookModes::ByTheBookInEditor: ELauncherProfileCookModes::ByTheBook;
 	}
 	else
 	{
 		bCookOnTheFly = true;
-		bInEditorCooking = Algo::AllOf(LauncherProfile->GetCookedPlatforms(),
+		bInEditorCooking = Algo::AllOf(BuildCookRun->GetCookedPlatforms(),
 			[this](const FString& PlatformName) { return CanCookOnTheFlyInEditor(PlatformName); });
 		CurrentLauncherCookMode = bInEditorCooking ? ELauncherProfileCookModes::OnTheFlyInEditor : ELauncherProfileCookModes::OnTheFly;
 	}
@@ -272,7 +273,7 @@ void UEditorEngine::StartPlayUsingLauncherSession(FRequestPlaySessionParams& InR
 	}
 
 	TStringBuilder<256> CookOptions;
-	CookOptions << LauncherProfile->GetCookOptions();
+	CookOptions << BuildCookRun->GetCookOptions();
 	ensure(CookOptions.Len() == 0);
 	auto SetCookOption = [&CookOptions](FStringView Option, bool bOptionOn)
 	{
@@ -292,26 +293,26 @@ void UEditorEngine::StartPlayUsingLauncherSession(FRequestPlaySessionParams& InR
 		const FTargetInfo* TargetInfo = GetDefault<UPlatformsMenuSettings>()->GetLaunchOnTargetInfo();
 		if (TargetInfo != nullptr)
 		{
-			LauncherProfile->ClearBuildTargets();
-			LauncherProfile->AddBuildTarget(TargetInfo->Name);
-			LauncherProfile->SetBuildTargetSpecified(true);
+			BuildCookRun->ClearBuildTargets();
+			BuildCookRun->AddBuildTarget(TargetInfo->Name);
+			BuildCookRun->SetBuildTargetSpecified(true);
 		}
 	}
-	LauncherProfile->SetCookMode(CurrentLauncherCookMode);
-	LauncherProfile->SetUnversionedCooking(!bIncrementalCooking || !bLegacyIterative); // Unversioned cooking is not allowed with legacy incremental cooking
-	LauncherProfile->SetIncrementalCooking(bIncrementalCooking);
+	BuildCookRun->SetCookMode(CurrentLauncherCookMode);
+	BuildCookRun->SetUnversionedCooking(!bIncrementalCooking || !bLegacyIterative); // Unversioned cooking is not allowed with legacy incremental cooking
+	BuildCookRun->SetIncrementalCooking(bIncrementalCooking);
 	SetCookOption(TEXTVIEW("-IgnoreIniSettingsOutOfDate"), bIncrementalCooking && CookerSettings.bIgnoreIniSettingsOutOfDateForIteration);
 	SetCookOption(TEXTVIEW("-IgnoreScriptPackagesOutOfDate"), bIncrementalCooking && CookerSettings.bIgnoreScriptPackagesOutOfDateForIteration);
 	SetCookOption(TEXTVIEW("-IterateSharedCookedbuild"), bIncrementalCooking && ExperimentalSettings.bSharedCookedBuilds);
-	LauncherProfile->SetDeployedDeviceGroup(DeviceGroup);
-	LauncherProfile->SetIncrementalDeploying(bIncrementalCooking);
+	BuildCookRun->SetDeployedDeviceGroup(DeviceGroup);
+	BuildCookRun->SetIncrementalDeploying(bIncrementalCooking);
 	LauncherProfile->SetEditorExe(FUnrealEdMisc::Get().GetExecutableForCommandlets());
 	LauncherProfile->SetShouldUpdateDeviceFlash(InRequestParams.LauncherTargetDevice->bUpdateDeviceFlash);
-	LauncherProfile->SetCookOptions(*CookOptions);
+	BuildCookRun->SetCookOptions(*CookOptions);
 
 	if (!InRequestParams.LauncherTargetDevice->Architecture.IsEmpty())
 	{
-		LauncherProfile->SetClientArchitectures({&InRequestParams.LauncherTargetDevice->Architecture, 1});
+		BuildCookRun->SetClientArchitectures({&InRequestParams.LauncherTargetDevice->Architecture, 1});
 	}
 	
 	if (LauncherProfile->IsBuildingUAT() && !GetDefault<UEditorPerProjectUserSettings>()->bAlwaysBuildUAT && bUATSuccessfullyCompiledOnce)
@@ -326,11 +327,11 @@ void UEditorEngine::StartPlayUsingLauncherSession(FRequestPlaySessionParams& InR
 	if ((LaunchPlatformName != TEXT("IOS") && LaunchPlatformName != TEXT("TVOS")) ||
 		(!LauncherSessionInfo->PlayUsingLauncherDeviceName.Contains(DummyIOSDeviceName) && !LauncherSessionInfo->PlayUsingLauncherDeviceName.Contains(DummyTVOSDeviceName)))
 	{
-		LauncherProfile->SetLaunchMode(ELauncherProfileLaunchModes::DefaultRole);
+		BuildCookRun->SetLaunchMode(ELauncherProfileLaunchModes::DefaultRole);
 	}
 
 	const bool bUseZenStore = GetDefault<UProjectPackagingSettings>()->GetUseZenStoreEffective();
-	LauncherProfile->SetUseZenStore(bUseZenStore);
+BuildCookRun->SetUseZenStore(bUseZenStore);
 #if UE_WITH_ZEN
 	if (bUseZenStore)
 	{
@@ -338,9 +339,9 @@ void UEditorEngine::StartPlayUsingLauncherSession(FRequestPlaySessionParams& InR
 	}
 #endif
 
-	if (bUseZenStore || LauncherProfile->GetCookMode() == ELauncherProfileCookModes::OnTheFlyInEditor || LauncherProfile->GetCookMode() == ELauncherProfileCookModes::OnTheFly)
+	if (bUseZenStore || BuildCookRun->GetCookMode() == ELauncherProfileCookModes::OnTheFlyInEditor || BuildCookRun->GetCookMode() == ELauncherProfileCookModes::OnTheFly)
 	{
-		LauncherProfile->SetDeploymentMode(ELauncherProfileDeploymentModes::FileServer);
+		BuildCookRun->SetDeploymentMode(ELauncherProfileDeploymentModes::FileServer);
 	}
 
 	switch(EditorPlaySettings->PackFilesForLaunch)
@@ -349,12 +350,12 @@ void UEditorEngine::StartPlayUsingLauncherSession(FRequestPlaySessionParams& InR
 	case EPlayOnPakFileMode::NoPak:
 		break;
 	case EPlayOnPakFileMode::PakNoCompress:
-		LauncherProfile->SetCompressed( false );
-		LauncherProfile->SetDeployWithUnrealPak( true );
+		BuildCookRun->SetCompressed( false );
+		BuildCookRun->SetDeployWithUnrealPak( true );
 		break;
 	case EPlayOnPakFileMode::PakCompress:
-		LauncherProfile->SetCompressed( true );
-		LauncherProfile->SetDeployWithUnrealPak( true );
+		BuildCookRun->SetCompressed( true );
+		BuildCookRun->SetDeployWithUnrealPak( true );
 		break;
 	}
 
@@ -375,24 +376,24 @@ void UEditorEngine::StartPlayUsingLauncherSession(FRequestPlaySessionParams& InR
 		InitialMapName = MapNames[0];
 	}
 
-	LauncherProfile->GetDefaultLaunchRole()->SetInitialMap(InitialMapName);
+	BuildCookRun->GetDefaultLaunchRole()->SetInitialMap(InitialMapName);
 
 	for (const FString& MapName : MapNames)
 	{
-		LauncherProfile->AddCookedMap(MapName);
+		BuildCookRun->AddCookedMap(MapName);
 	}
 
-	if (LauncherProfile->GetCookMode() == ELauncherProfileCookModes::ByTheBookInEditor)
+	if (BuildCookRun->GetCookMode() == ELauncherProfileCookModes::ByTheBookInEditor)
 	{
 		TArray<ITargetPlatform*> TargetPlatforms;
-		for (const FString& PlatformName : LauncherProfile->GetCookedPlatforms())
+		for (const FString& PlatformName : BuildCookRun->GetCookedPlatforms())
 		{
 			ITargetPlatform* TargetPlatform = GetTargetPlatformManager()->FindTargetPlatform(PlatformName);
 			// todo pass in all the target platforms instead of just the single platform
 			// crashes if two requests are inflight but we can support having multiple platforms cooking at once
 			TargetPlatforms.Add(TargetPlatform);
 		}
-		const TArray<FString> &CookedMaps = LauncherProfile->GetCookedMaps();
+		const TArray<FString> &CookedMaps = BuildCookRun->GetCookedMaps();
 
 		// const TArray<FString>& CookedMaps = ChainState.Profile->GetCookedMaps();
 		TArray<FString> CookDirectories;
diff --git a/Engine/Source/Editor/UnrealEd/Private/Settings/ProjectPackagingSettingsCustomization.h b/Engine/Source/Editor/UnrealEd/Private/Settings/ProjectPackagingSettingsCustomization.h
--- a/Engine/Source/Editor/UnrealEd/Private/Settings/ProjectPackagingSettingsCustomization.h
+++ b/Engine/Source/Editor/UnrealEd/Private/Settings/ProjectPackagingSettingsCustomization.h
@@ -490,6 +490,11 @@ protected:
 		ILauncherServicesModule& LauncherServicesModule = FModuleManager::LoadModuleChecked<ILauncherServicesModule>("LauncherServices");
 		for (const ILauncherProfilePtr& LauncherProfilePtr : LauncherServicesModule.GetProfileManager()->GetAllProfiles())
 		{
+			if (!LauncherProfilePtr->GetFirstBuildCookRun().IsValid())
+			{
+				continue;
+			}
+
 			MenuBuilder.AddMenuEntry(
 				FText::FromString(LauncherProfilePtr->GetName()),
 				FText::FromString(LauncherProfilePtr->GetDescription()),
@@ -509,6 +514,12 @@ protected:
 	{
 		ILauncherServicesModule& LauncherServicesModule = FModuleManager::LoadModuleChecked<ILauncherServicesModule>("LauncherServices");
 
+		const ILauncherProfileBuildCookRunPtr BuildCookRunPtr = LauncherProfilePtr->GetFirstBuildCookRun();
+		if (BuildCookRunPtr == nullptr)
+		{
+			return;
+		}
+
 		// grab the project packaging settings
 		UProjectPackagingSettings* ProjectPackagingSettings = UProjectPackagingSettings::StaticClass()->GetDefaultObject<UProjectPackagingSettings>();
 		if (ProjectPackagingSettings == nullptr)
@@ -528,10 +539,10 @@ protected:
 		FProjectBuildSettings& ProjectBuildSettings = ProjectPackagingSettings->ProjectCustomBuilds.AddDefaulted_GetRef();
 		ProjectBuildSettings.Name               = ProfileName;
 		ProjectBuildSettings.HelpText           = LauncherProfilePtr->GetDescription().IsEmpty() ? LauncherProfilePtr->GetName() : LauncherProfilePtr->GetDescription();
-		ProjectBuildSettings.SpecificPlatforms  = LauncherProfilePtr->GetCookedPlatforms();
-		if (ProjectBuildSettings.SpecificPlatforms.Num() == 0 && !LauncherProfilePtr->GetDefaultDeployPlatform().IsNone())
+		ProjectBuildSettings.SpecificPlatforms  = BuildCookRunPtr->GetCookedPlatforms();
+		if (ProjectBuildSettings.SpecificPlatforms.Num() == 0 && !BuildCookRunPtr->GetDefaultDeployPlatform().IsNone())
 		{
-			ProjectBuildSettings.SpecificPlatforms.Add(LauncherProfilePtr->GetDefaultDeployPlatform().ToString());
+			ProjectBuildSettings.SpecificPlatforms.Add(BuildCookRunPtr->GetDefaultDeployPlatform().ToString());
 		}
 		ProjectBuildSettings.BuildCookRunParams = LauncherServicesModule.GetProfileManager()->MakeBuildCookRunParamsForProjectCustomBuild(LauncherProfilePtr.ToSharedRef(), ProjectBuildSettings.SpecificPlatforms);
 
