bug fix - append the launch command line
(lost during the multi-build/cook/run refactor)

#rnx

#Synced-CL 50434826

diff --git a/Engine/Source/Developer/LauncherServices/Private/Launcher/LauncherWorker.cpp b/Engine/Source/Developer/LauncherServices/Private/Launcher/LauncherWorker.cpp
--- a/Engine/Source/Developer/LauncherServices/Private/Launcher/LauncherWorker.cpp
+++ b/Engine/Source/Developer/LauncherServices/Private/Launcher/LauncherWorker.cpp
@@ -681,6 +681,8 @@ FString FLauncherWorker::CreateUATCommandForBuildCookRun( const ILauncherProfile
 				break;
 			}
 		}
+
+		UATCommand += AdditionalCommandLine;
 	}
 
 
@@ -1073,7 +1075,6 @@ FString FLauncherWorker::CreateUATCommandForBuildCookRun( const ILauncherProfile
 				UATCommand += CommandLine;
 				UATCommand += StageDirectory;
 				UATCommand += DeviceCommand;
-				//UATCommand += AdditionalCommandLine;
 
 				FCommandDesc Desc;
 				FText Command = FText::Format(LOCTEXT("LauncherDeployDesc", "Deploying content for {0}"), FText::FromString(FString::Join(FinalPlatforms, TEXT("+") )));
@@ -1101,7 +1102,6 @@ FString FLauncherWorker::CreateUATCommandForBuildCookRun( const ILauncherProfile
 				UATCommand += CommandLine;
 				UATCommand += StageDirectory;
 				UATCommand += DeviceCommand;
-				//UATCommand += AdditionalCommandLine;
 				UATCommand += StageAdditionalCommandLine;
 
 				FCommandDesc Desc;
@@ -1130,7 +1130,6 @@ FString FLauncherWorker::CreateUATCommandForBuildCookRun( const ILauncherProfile
 			UATCommand += TEXT(" -stage -package");
 			UATCommand += StageDirectory;
 			UATCommand += CommandLine;
-			//UATCommand += AdditionalCommandLine;
 			UATCommand += StageAdditionalCommandLine;
 
 			FCommandDesc Desc;
