fixes for target device selection not being saved correctly

- device groups are saved in a common location
- new device groups are only assigned when a profile is created, or doesn't have one already
- device groups list is trimmed and saved on shutdown
- empty device ids are not saved

#rnx

#Synced-CL 50642529

diff --git a/Engine/Plugins/Developer/ProjectLauncher/Source/ProjectLauncher/Private/Model/ProjectLauncherModel.cpp b/Engine/Plugins/Developer/ProjectLauncher/Source/ProjectLauncher/Private/Model/ProjectLauncherModel.cpp
--- a/Engine/Plugins/Developer/ProjectLauncher/Source/ProjectLauncher/Private/Model/ProjectLauncherModel.cpp
+++ b/Engine/Plugins/Developer/ProjectLauncher/Source/ProjectLauncher/Private/Model/ProjectLauncherModel.cpp
@@ -114,6 +114,7 @@ namespace ProjectLauncher
 		DeviceProxyManager->OnProxyRemoved().RemoveAll(this);
 
 		SaveConfig();
+		ProfileManager->SaveSettings();
 
 		ShutdownLaunchExtensions();
 	}
diff --git a/Engine/Source/Developer/LauncherServices/Private/Profiles/LauncherProfile.h b/Engine/Source/Developer/LauncherServices/Private/Profiles/LauncherProfile.h
--- a/Engine/Source/Developer/LauncherServices/Private/Profiles/LauncherProfile.h
+++ b/Engine/Source/Developer/LauncherServices/Private/Profiles/LauncherProfile.h
@@ -1676,8 +1676,6 @@ public:
 		UATCommandValidation.Reset();
 		FirstBuildCookRun = CreateBuildCookRunInternal(TEXT("FirstBuildCookRun"), TEXT(""));
 		AddUATCommandInternal(FirstBuildCookRun.ToSharedRef());
-		FirstBuildCookRun->SetDeployedDeviceGroup(LauncherProfileManager->AddNewDeviceGroup());
-
 
 		RefreshValidBuildTargets();
 		Validate();
diff --git a/Engine/Source/Developer/LauncherServices/Private/Profiles/LauncherProfileBuildCookRun.h b/Engine/Source/Developer/LauncherServices/Private/Profiles/LauncherProfileBuildCookRun.h
--- a/Engine/Source/Developer/LauncherServices/Private/Profiles/LauncherProfileBuildCookRun.h
+++ b/Engine/Source/Developer/LauncherServices/Private/Profiles/LauncherProfileBuildCookRun.h
@@ -114,7 +114,7 @@ public:
 	 *
 	 * @return The device group identifier, or an invalid GUID if no group was set or deployment is disabled.
 	 */
-	const FGuid& GetDeployedDeviceGroupId( ) const
+	virtual const FGuid& GetDeployedDeviceGroupId( ) const override
 	{
 		return DeployedDeviceGroupId;
 	}
diff --git a/Engine/Source/Developer/LauncherServices/Private/Profiles/LauncherProfileManager.cpp b/Engine/Source/Developer/LauncherServices/Private/Profiles/LauncherProfileManager.cpp
--- a/Engine/Source/Developer/LauncherServices/Private/Profiles/LauncherProfileManager.cpp
+++ b/Engine/Source/Developer/LauncherServices/Private/Profiles/LauncherProfileManager.cpp
@@ -4,6 +4,7 @@
 #include "HAL/FileManager.h"
 #include "Misc/FileHelper.h"
 #include "Misc/ConfigCacheIni.h"
+#include "Misc/ConfigContext.h"
 #include "Serialization/JsonReader.h"
 #include "Serialization/JsonSerializer.h"
 #include "Launcher/LauncherProjectPath.h"
@@ -56,6 +57,8 @@ LAUNCHERSERVICES_API bool HasPromotedTarget(const TCHAR* BaseDir, const TCHAR* T
 
 FLauncherProfileManager::FLauncherProfileManager()
 {
+	ConfigFileName = FPaths::EngineSavedDir() / TEXT("Config") / TEXT("ProjectLauncher") / TEXT("UserSettings.ini");
+	FConfigContext::ReadSingleIntoGConfig().Load(*ConfigFileName);
 }
 
 /* ILauncherProfileManager interface
@@ -304,6 +307,16 @@ ILauncherProfilePtr FLauncherProfileManager::LoadJSONProfile(FString ProfileFile
 
 	if (Profile->Load(*(Object.Get())))
 	{
+		for (const ILauncherProfileBuildCookRunRef& BuildCookRun : Profile->GetBuildCookRunCommands())
+		{
+			ILauncherDeviceGroupPtr DeviceGroup = GetDeviceGroup(BuildCookRun->GetDeployedDeviceGroupId());
+			if (!DeviceGroup.IsValid())
+			{
+				DeviceGroup = AddNewDeviceGroup();
+			}
+			BuildCookRun->SetDeployedDeviceGroup(DeviceGroup);
+		}
+
 		return MakeShareable(Profile);
 	}
 
@@ -435,7 +448,7 @@ const TArray<ILauncherProfileWizardPtr>& FLauncherProfileManager::GetProfileWiza
 
 void FLauncherProfileManager::SaveSettings( )
 {
-	SaveDeviceGroups();
+	SaveDeviceGroups(true);
 	SaveSimpleProfiles();
 	SaveProfiles();
 }
@@ -509,16 +522,26 @@ void FLauncherProfileManager::LoadDeviceGroups( )
 {
 	if (GConfig != nullptr)
 	{
-		const FConfigSection* LoadedDeviceGroups = GConfig->GetSection(TEXT("Launcher.DeviceGroups"), false, GEngineIni);
-
-		if (LoadedDeviceGroups != nullptr)
+		for (int Pass = 0; Pass < 2; Pass++)
 		{
-			// parse the configuration file entries into device groups
-			for (FConfigSection::TConstIterator It(*LoadedDeviceGroups); It; ++It)
+			const FConfigSection* LoadedDeviceGroups = (Pass == 0)
+				? GConfig->GetSection(TEXT("Launcher.DeviceGroups"), false, GEngineIni)     // legacy location (differs per-project & per-program)
+				: GConfig->GetSection(TEXT("Launcher.DeviceGroups"), false, ConfigFileName) // new location (common)
+			;
+
+			if (LoadedDeviceGroups != nullptr)
 			{
-				if (It.Key() == TEXT("DeviceGroup"))
+				// parse the configuration file entries into device groups
+				for (FConfigSection::TConstIterator It(*LoadedDeviceGroups); It; ++It)
 				{
-					DeviceGroups.Add(ParseDeviceGroup(It.Value().GetValue()));
+					if (It.Key() == TEXT("DeviceGroup"))
+					{
+						ILauncherDeviceGroupPtr DeviceGroup = ParseDeviceGroup(It.Value().GetValue());
+						if (DeviceGroup.IsValid())
+						{
+							DeviceGroups.Add(DeviceGroup);
+						}
+					}
 				}
 			}
 		}
@@ -642,7 +665,7 @@ ILauncherDeviceGroupPtr FLauncherProfileManager::ParseDeviceGroup( const FString
 		Result = MakeShareable(new FLauncherDeviceGroup(GroupId, GroupName));
 
 		TArray<FString> DeviceList;
-		DevicesString.ParseIntoArray(DeviceList, TEXT(", "), false);
+		DevicesString.ParseIntoArray(DeviceList, TEXT(", "), true);
 
 		for (int32 Index = 0; Index < DeviceList.Num(); ++Index)
 		{
@@ -654,18 +677,38 @@ ILauncherDeviceGroupPtr FLauncherProfileManager::ParseDeviceGroup( const FString
 }
 
 
-void FLauncherProfileManager::SaveDeviceGroups( )
+void FLauncherProfileManager::SaveDeviceGroups( bool bTrimUnused )
 {
 	if (GConfig != nullptr)
 	{
-		GConfig->EmptySection(TEXT("Launcher.DeviceGroups"), GEngineIni);
+		TSet<FGuid> ActiveDeviceGroupIds;
+		for ( const ILauncherProfilePtr& LauncherProfile : SavedProfiles )
+		{
+			for (const ILauncherProfileBuildCookRunRef& BuildCookRun : LauncherProfile->GetBuildCookRunCommands())
+			{
+				const ILauncherDeviceGroupPtr& Group = BuildCookRun->GetDeployedDeviceGroup(false);
+				if (Group.IsValid())
+				{
+					ActiveDeviceGroupIds.Add( Group->GetId() );
+				}
+			}
+		}
+
+
+		GConfig->EmptySection(TEXT("Launcher.DeviceGroups"), ConfigFileName);
 
 		TArray<FString> DeviceGroupStrings;
 
+
 		// create a string representation of all groups and their devices
 		for (int32 GroupIndex = 0; GroupIndex < DeviceGroups.Num(); ++GroupIndex)
 		{
 			const ILauncherDeviceGroupPtr& Group = DeviceGroups[GroupIndex];
+			if (bTrimUnused && !ActiveDeviceGroupIds.Contains(Group->GetId()))
+			{
+				continue; // trim out unused groups when saving
+			}
+
 			const TArray<FString>& Devices = Group->GetDeviceIDs();
 
 			FString DeviceListString;
@@ -686,8 +729,8 @@ void FLauncherProfileManager::SaveDeviceGroups( )
 		}
 
 		// save configuration
-		GConfig->SetArray(TEXT("Launcher.DeviceGroups"), TEXT("DeviceGroup"), DeviceGroupStrings, GEngineIni);
-		GConfig->Flush(false, GEngineIni);
+		GConfig->SetArray(TEXT("Launcher.DeviceGroups"), TEXT("DeviceGroup"), DeviceGroupStrings, ConfigFileName);
+		GConfig->Flush(false, ConfigFileName);
 	}
 }
 
diff --git a/Engine/Source/Developer/LauncherServices/Private/Profiles/LauncherProfileManager.h b/Engine/Source/Developer/LauncherServices/Private/Profiles/LauncherProfileManager.h
--- a/Engine/Source/Developer/LauncherServices/Private/Profiles/LauncherProfileManager.h
+++ b/Engine/Source/Developer/LauncherServices/Private/Profiles/LauncherProfileManager.h
@@ -104,7 +104,7 @@ public:
 
 	virtual void RemoveDeviceGroup( const ILauncherDeviceGroupRef& DeviceGroup ) override;
 
-	virtual void SaveDeviceGroups() override;
+	virtual void SaveDeviceGroups( bool bTrimUnused = false ) override;
 
 	virtual void RemoveSimpleProfile(const ILauncherSimpleProfileRef& SimpleProfile) override;
 
@@ -212,6 +212,9 @@ private:
 	// Holds all registered profile wizards
 	TArray<ILauncherProfileWizardPtr> ProfileWizards;
 
+	// Holds the config file name to save to
+	FString ConfigFileName;
+
 private:
 
 	// Holds a delegate to be invoked when a device group was added.
diff --git a/Engine/Source/Developer/LauncherServices/Public/ILauncherProfileBuildCookRun.h b/Engine/Source/Developer/LauncherServices/Public/ILauncherProfileBuildCookRun.h
--- a/Engine/Source/Developer/LauncherServices/Public/ILauncherProfileBuildCookRun.h
+++ b/Engine/Source/Developer/LauncherServices/Public/ILauncherProfileBuildCookRun.h
@@ -318,6 +318,17 @@ public:
 	/** Sets all profile settings to their defaults. */
 	virtual void SetDefaults( ) = 0;
 
+	/**
+	 * Gets the identifier of the device group to deploy to.
+	 *
+	 * This method is used internally by the profile manager to read the device group identifier after
+	 * loading this profile from a file. The profile manager will use this identifier to locate the
+	 * actual device group to deploy to.
+	 *
+	 * @return The device group identifier, or an invalid GUID if no group was set or deployment is disabled.
+	 */
+	virtual const FGuid& GetDeployedDeviceGroupId() const = 0;
+
 public:
 
 	/**
diff --git a/Engine/Source/Developer/LauncherServices/Public/ILauncherProfileManager.h b/Engine/Source/Developer/LauncherServices/Public/ILauncherProfileManager.h
--- a/Engine/Source/Developer/LauncherServices/Public/ILauncherProfileManager.h
+++ b/Engine/Source/Developer/LauncherServices/Public/ILauncherProfileManager.h
@@ -169,7 +169,7 @@ public:
 	virtual void RemoveDeviceGroup( const ILauncherDeviceGroupRef& DeviceGroup ) = 0;
 
 	/** Saves all the device groups to a config file */
-	virtual void SaveDeviceGroups() = 0;
+	virtual void SaveDeviceGroups( bool bTrimUnused = false ) = 0;
 
 public:
 
