fix for removing a launched-on device when the device variant doesn't match the deploy device group

#rnx

#Synced-CL 50642529

diff --git a/Engine/Source/Developer/LauncherServices/Private/Launcher/LauncherWorker.cpp b/Engine/Source/Developer/LauncherServices/Private/Launcher/LauncherWorker.cpp
--- a/Engine/Source/Developer/LauncherServices/Private/Launcher/LauncherWorker.cpp
+++ b/Engine/Source/Developer/LauncherServices/Private/Launcher/LauncherWorker.cpp
@@ -1859,7 +1859,14 @@ void FLauncherWorker::CreateAndExecuteTasks( const ILauncherProfileRef& InProfil
 /** Callback for when a device proxy has been removed from the device proxy manager. */
 void FLauncherWorker::HandleDeviceProxyManagerProxyRemoved(const TSharedRef<ITargetDeviceProxy>& RemovedProxy, ILauncherProfileUATCommandPtr UnifiedUATCommand)
 {
-	FString TargetDeviceId = RemovedProxy->GetTargetDeviceId(NAME_None);
+	TArray<FName> Variants;
+	RemovedProxy->GetVariants(Variants);
+
+	TSet<FString> TargetDeviceIds;
+	for (FName Variant : Variants)
+	{
+		TargetDeviceIds.Append( RemovedProxy->GetTargetDeviceIds(Variant) );
+	}
 
 	TArray<ILauncherProfileUATCommandRef> UATCommands = Profile->GetUATCommands();
 	if (UnifiedUATCommand.IsValid())
@@ -1883,12 +1890,15 @@ void FLauncherWorker::HandleDeviceProxyManagerProxyRemoved(const TSharedRef<ITar
 				if (DeviceGroup.IsValid() && DeviceGroup->GetNumDevices() > 0)
 				{
 					// remove disconnected device from the list
-					DeviceGroup->RemoveDevice(TargetDeviceId);
-					if (DeviceGroup->GetNumDevices() == 0)
+					for (const FString& TargetDeviceId : TargetDeviceIds)
 					{
-						// this was the last device running, stop working
-						Stop();
-						break;
+						DeviceGroup->RemoveDevice(TargetDeviceId);
+						if (DeviceGroup->GetNumDevices() == 0)
+						{
+							// this was the last device running, stop working
+							Stop();
+							break;
+						}
 					}
 				}
 			}
