shutdown fix for project launcher saving profiles on exit

#rnx

#Synced-CL 50642529

diff --git a/Engine/Plugins/Developer/ProjectLauncher/Source/ProjectLauncher/Private/Model/ProjectLauncherModel.cpp b/Engine/Plugins/Developer/ProjectLauncher/Source/ProjectLauncher/Private/Model/ProjectLauncherModel.cpp
--- a/Engine/Plugins/Developer/ProjectLauncher/Source/ProjectLauncher/Private/Model/ProjectLauncherModel.cpp
+++ b/Engine/Plugins/Developer/ProjectLauncher/Source/ProjectLauncher/Private/Model/ProjectLauncherModel.cpp
@@ -114,7 +114,7 @@ namespace ProjectLauncher
 		DeviceProxyManager->OnProxyRemoved().RemoveAll(this);
 
 		SaveConfig();
-		ProfileManager->SaveSettings();
+		// ProfileManager->SaveSettings(); // not doing this at the moment because it also saves all the profiles which seems unnecessary
 
 		ShutdownLaunchExtensions();
 	}
diff --git a/Engine/Source/Developer/LauncherServices/Private/Profiles/LauncherProfile.h b/Engine/Source/Developer/LauncherServices/Private/Profiles/LauncherProfile.h
--- a/Engine/Source/Developer/LauncherServices/Private/Profiles/LauncherProfile.h
+++ b/Engine/Source/Developer/LauncherServices/Private/Profiles/LauncherProfile.h
@@ -1436,8 +1436,12 @@ public:
 		FName Variant = NAME_None;
 
 		// Loading the Device Proxy Manager to get the needed Device Manager.
-		ITargetDeviceServicesModule& DeviceServiceModule = FModuleManager::LoadModuleChecked<ITargetDeviceServicesModule>(TEXT("TargetDeviceServices"));
-		TSharedRef<ITargetDeviceProxyManager> DeviceProxyManager = DeviceServiceModule.GetDeviceProxyManager();
+		ITargetDeviceServicesModule* DeviceServiceModule = FModuleManager::LoadModulePtr<ITargetDeviceServicesModule>(TEXT("TargetDeviceServices"));
+		if (DeviceServiceModule == nullptr)
+		{
+			return Platforms;
+		}
+		TSharedRef<ITargetDeviceProxyManager> DeviceProxyManager = DeviceServiceModule->GetDeviceProxyManager();
 
 		if (DeviceGroup.IsValid() && Platforms.Num() < 1)
 		{
