bug fix for uat command line not being updated by platform extensions

#rnx

#Synced-CL 50857758

diff --git a/Engine/Source/Developer/LauncherServices/Private/Launcher/LauncherWorker.cpp b/Engine/Source/Developer/LauncherServices/Private/Launcher/LauncherWorker.cpp
--- a/Engine/Source/Developer/LauncherServices/Private/Launcher/LauncherWorker.cpp
+++ b/Engine/Source/Developer/LauncherServices/Private/Launcher/LauncherWorker.cpp
@@ -1194,7 +1194,7 @@ FString FLauncherWorker::CreateUATCommandForBuildCookRun( const ILauncherProfile
 		.bShouldBeSkipped = false, // don't skip by default
 		.Profile = InProfile,
 		.UATCommand = InBuildCookRun,
-		.UATCommandLine = AdditionalUATCommandLine,
+		.UATCommandLine = UATCommand,
 		.Commands = OutCommands,
 		.CommandStart = CommandStart,
 	};
