Build Target optimizations for UFE to avoid spawning UBT if it's not necessary
- always clearing the fallback build target when the project changes, rather than trying to maintain it between projects - is it highly unlikely to ever be the same
- profile's build target validations only done when using the legacy project launcher
- mark 'required build target' functions as deprecated because they're only used by the legacy project launcher

#rnx

#Synced-CL 51007224

diff --git a/Engine/Source/Developer/LauncherServices/Private/Profiles/LauncherProfile.h b/Engine/Source/Developer/LauncherServices/Private/Profiles/LauncherProfile.h
--- a/Engine/Source/Developer/LauncherServices/Private/Profiles/LauncherProfile.h
+++ b/Engine/Source/Developer/LauncherServices/Private/Profiles/LauncherProfile.h
@@ -2803,6 +2803,12 @@ protected:
 
 	void ValidateBuildTarget(const ILauncherProfileBuildCookRunRef& BuildCookRun)
 	{
+		// this is only useful if the legacy project launcher is being used
+		if (!BuildTargetOptionsChangedDelegate.IsBound())
+		{
+			return;
+		}
+	
 		FValidationData& ValidationData = UATCommandValidation.FindChecked(BuildCookRun.ToSharedPtr());
 
 		bool bBuildTargetIsRequired = false;
@@ -2830,8 +2836,10 @@ protected:
 		}
 
 		// check all build targets
+PRAGMA_DISABLE_DEPRECATION_WARNINGS
 		const TArray<FString>& BuildTargetNames = HasProjectSpecified() ? ExplictBuildTargetNames : LauncherProfileManager->GetAllExplicitBuildTargetNames();	
 		const TArray<FTargetInfo>& Targets = FDesktopPlatformModule::Get()->GetTargetsForProject(GetProjectPath());
+PRAGMA_ENABLE_DEPRECATION_WARNINGS
 		for (const FTargetInfo& Target : Targets)
 		{
 			if (CookTargetTypes.Contains(Target.Type))
@@ -2886,6 +2894,12 @@ protected:
 
 	void RefreshValidBuildTargets()
 	{
+		// this is only useful if the legacy project launcher is being used
+		if (!BuildTargetOptionsChangedDelegate.IsBound())
+		{
+			return;
+		}
+		
 		TArray<FString> LatestExplicitBuildTargetNames;
 
 		// collect the build targets for the current project, filtered to what we are currently wanting to cook. Do not show fallback project's build targets
diff --git a/Engine/Source/Developer/LauncherServices/Private/Profiles/LauncherProfileManager.cpp b/Engine/Source/Developer/LauncherServices/Private/Profiles/LauncherProfileManager.cpp
--- a/Engine/Source/Developer/LauncherServices/Private/Profiles/LauncherProfileManager.cpp
+++ b/Engine/Source/Developer/LauncherServices/Private/Profiles/LauncherProfileManager.cpp
@@ -473,12 +473,8 @@ void FLauncherProfileManager::SetProjectPath(const FString& InProjectPath)
 	if (ProjectPath != InProjectPath)
 	{
 		ProjectPath = InProjectPath;
-		ExplicitBuildTargets = FGameProjectHelper::GetExplicitBuildTargetsForProject(ProjectPath);
 
-		if (!BuildTarget.IsEmpty() && !ExplicitBuildTargets.Contains(BuildTarget))
-		{
-			SetBuildTarget(FString());
-		}
+		SetBuildTarget(FString());
 
 		for (ILauncherProfilePtr Profile : AllProfiles)
 		{
@@ -512,9 +508,9 @@ void FLauncherProfileManager::SetBuildTarget( const FString& InBuildTarget )
 	}
 }
 
-const TArray<FString>& FLauncherProfileManager::GetAllExplicitBuildTargetNames() const
+const TArray<FString> FLauncherProfileManager::GetAllExplicitBuildTargetNames() const
 {
-	return ExplicitBuildTargets;
+	return FGameProjectHelper::GetExplicitBuildTargetsForProject(ProjectPath);
 }
 
 
diff --git a/Engine/Source/Developer/LauncherServices/Private/Profiles/LauncherProfileManager.h b/Engine/Source/Developer/LauncherServices/Private/Profiles/LauncherProfileManager.h
--- a/Engine/Source/Developer/LauncherServices/Private/Profiles/LauncherProfileManager.h
+++ b/Engine/Source/Developer/LauncherServices/Private/Profiles/LauncherProfileManager.h
@@ -137,7 +137,8 @@ public:
 
 	virtual void SetBuildTarget( const FString& InBuildTarget ) override;
 
-	virtual const TArray<FString>& GetAllExplicitBuildTargetNames() const override;
+	UE_DEPRECATED(5.8, "this was only used for the legacy project launcher")
+	virtual const TArray<FString> GetAllExplicitBuildTargetNames() const override;
 
 	virtual FString MakeBuildCookRunParamsForProjectCustomBuild(const ILauncherProfileRef& InProfile, const TArray<FString>& InPlatforms) const override;
 
@@ -205,9 +206,6 @@ private:
 
 	// Holds the currently selected build target
 	FString BuildTarget;
-
-	// Holds the current explicit build targets
-	TArray<FString> ExplicitBuildTargets;
 	
 	// Holds all registered profile wizards
 	TArray<ILauncherProfileWizardPtr> ProfileWizards;
diff --git a/Engine/Source/Developer/LauncherServices/Public/ILauncherProfile.h b/Engine/Source/Developer/LauncherServices/Public/ILauncherProfile.h
--- a/Engine/Source/Developer/LauncherServices/Public/ILauncherProfile.h
+++ b/Engine/Source/Developer/LauncherServices/Public/ILauncherProfile.h
@@ -1372,6 +1372,7 @@ public:
 	 * 
 	 * @return The delegate.
 	 */
+	UE_DEPRECATED(5.8, "this was only used for the legacy project launcher")
 	virtual FOnProfileBuildTargetOptionsChanged& OnBuildTargetOptionsChanged() = 0;
 
 	/**
@@ -1687,11 +1688,13 @@ public:
 	/**
 	 * Helper function to get all of the build targets available for this profile, based on its current project & cook platforms
 	 */
+	UE_DEPRECATED(5.8, "this was only used for the legacy project launcher")
 	virtual TArray<FString> GetExplicitBuildTargetNames() const = 0;
 
 	/**
 	 * Whether the profile will require an explicit -target=XXX parameter
 	 */
+	UE_DEPRECATED(5.8, "this was only used for the legacy project launcher")
 	virtual bool RequiresExplicitBuildTargetName() const = 0;
 
 	/**
diff --git a/Engine/Source/Developer/LauncherServices/Public/ILauncherProfileManager.h b/Engine/Source/Developer/LauncherServices/Public/ILauncherProfileManager.h
--- a/Engine/Source/Developer/LauncherServices/Public/ILauncherProfileManager.h
+++ b/Engine/Source/Developer/LauncherServices/Public/ILauncherProfileManager.h
@@ -395,7 +395,8 @@ public:
 	 * Return a list of all build build targets for the current project that need to be specified explicitly via -target=
 	 * If a build target isn't in the list, it does not need to be specified (and BuildTarget can be an empty string)
 	 */
-	virtual const TArray<FString>& GetAllExplicitBuildTargetNames() const = 0;
+	UE_DEPRECATED(5.8, "this was only used for the legacy project launcher")
+	virtual const TArray<FString> GetAllExplicitBuildTargetNames() const = 0;
 
 public:
 
diff --git a/Engine/Source/Developer/LegacyProjectLauncher/Private/Widgets/Shared/SProjectLauncherBuildTargetSelector.cpp b/Engine/Source/Developer/LegacyProjectLauncher/Private/Widgets/Shared/SProjectLauncherBuildTargetSelector.cpp
--- a/Engine/Source/Developer/LegacyProjectLauncher/Private/Widgets/Shared/SProjectLauncherBuildTargetSelector.cpp
+++ b/Engine/Source/Developer/LegacyProjectLauncher/Private/Widgets/Shared/SProjectLauncherBuildTargetSelector.cpp
@@ -27,6 +27,7 @@ SProjectLauncherBuildTargetSelector::~SProjectLauncherBuildTargetSelector()
 		Model->GetProfileManager()->OnProjectChanged().RemoveAll(this);
 	}
 
+PRAGMA_DISABLE_DEPRECATION_WARNINGS
 	if (bUseProfile)
 	{
 		ILauncherProfilePtr Profile = GetProfile();
@@ -36,6 +37,7 @@ SProjectLauncherBuildTargetSelector::~SProjectLauncherBuildTargetSelector()
 			Profile->OnBuildTargetOptionsChanged().RemoveAll(this);
 		}
 	}
+PRAGMA_ENABLE_DEPRECATION_WARNINGS
 }
 
 
