launcher profile's uat commands use shared pointers
- fixes a shutdown use-after-free

#rnx

#Synced-CL 51007224

diff --git a/Engine/Source/Developer/LauncherServices/Private/LauncherServicesModule.cpp b/Engine/Source/Developer/LauncherServices/Private/LauncherServicesModule.cpp
--- a/Engine/Source/Developer/LauncherServices/Private/LauncherServicesModule.cpp
+++ b/Engine/Source/Developer/LauncherServices/Private/LauncherServicesModule.cpp
@@ -40,7 +40,10 @@ public:
 	virtual ILauncherProfileRef CreateProfile(const FString& ProfileName) override
 	{
 		ILauncherProfileManagerRef ProfileManager = GetProfileManager();
-		return MakeShareable(new FLauncherProfile(ProfileManager, FGuid(), ProfileName));
+		TSharedRef<FLauncherProfile> NewProfile = MakeShared<FLauncherProfile>(ProfileManager, FGuid(), ProfileName);
+		NewProfile->SetDefaults();
+
+		return NewProfile;
 	}
 
 	virtual ILauncherProfileManagerRef GetProfileManager() override
diff --git a/Engine/Source/Developer/LauncherServices/Private/Profiles/LauncherProfile.cpp b/Engine/Source/Developer/LauncherServices/Private/Profiles/LauncherProfile.cpp
--- a/Engine/Source/Developer/LauncherServices/Private/Profiles/LauncherProfile.cpp
+++ b/Engine/Source/Developer/LauncherServices/Private/Profiles/LauncherProfile.cpp
@@ -16,19 +16,19 @@ const TCHAR* FLauncherProfileBuildCookRun::TypeName = TEXT("BuildCookRun");
 
 ILauncherProfileUATCommandRef FLauncherProfile::CreateUATCommandInternal(const TCHAR* InternalName, const TCHAR* UserTypeName)
 {
-	ILauncherProfileUATCommandRef UATCommand = MakeShared<FLauncherProfileUATCommand>(this, InternalName, UserTypeName);
+	ILauncherProfileUATCommandRef UATCommand = MakeShared<FLauncherProfileUATCommand>(AsShared(), InternalName, UserTypeName);
 	return UATCommand;
 }
 
 ILauncherProfileAutomatedTestRef FLauncherProfile::CreateAutomatedTestInternal(const TCHAR* InternalName, const TCHAR* UserTypeName)
 {
-	ILauncherProfileAutomatedTestRef AutomatedTest = MakeShared<FLauncherProfileAutomatedTest>(this, InternalName, UserTypeName);
+	ILauncherProfileAutomatedTestRef AutomatedTest = MakeShared<FLauncherProfileAutomatedTest>(AsShared(), InternalName, UserTypeName);
 	return AutomatedTest;
 }
 
 ILauncherProfileBuildCookRunRef FLauncherProfile::CreateBuildCookRunInternal(const TCHAR* InternalName, const TCHAR* UserTypeName)
 {
-	ILauncherProfileBuildCookRunRef BuildCookRun = MakeShared<FLauncherProfileBuildCookRun>(LauncherProfileManager, this, InternalName, UserTypeName);
+	ILauncherProfileBuildCookRunRef BuildCookRun = MakeShared<FLauncherProfileBuildCookRun>(LauncherProfileManager, AsShared(), InternalName, UserTypeName);
 	return BuildCookRun;
 
 }
@@ -57,7 +57,7 @@ ILauncherProfileUATCommandPtr FLauncherProfile::CreateUATCommandByTypeNameIntern
 
 ILauncherProfileBuildCookRunPtr FLauncherProfile::LoadBuildCookRunInternal(const FJsonObject& Object, int32 Version)
 {
-	TSharedRef<FLauncherProfileBuildCookRun> BuildCookRun = MakeShared<FLauncherProfileBuildCookRun>(LauncherProfileManager, this, TEXT("OriginalBuildCookRun"));
+	TSharedRef<FLauncherProfileBuildCookRun> BuildCookRun = MakeShared<FLauncherProfileBuildCookRun>(LauncherProfileManager, AsShared(), TEXT("OriginalBuildCookRun"));
 	if (!BuildCookRun->Load(Object, Version))
 	{
 		return nullptr;
diff --git a/Engine/Source/Developer/LauncherServices/Private/Profiles/LauncherProfile.h b/Engine/Source/Developer/LauncherServices/Private/Profiles/LauncherProfile.h
--- a/Engine/Source/Developer/LauncherServices/Private/Profiles/LauncherProfile.h
+++ b/Engine/Source/Developer/LauncherServices/Private/Profiles/LauncherProfile.h
@@ -200,7 +200,6 @@ public:
 	FLauncherProfile(ILauncherProfileManagerRef ProfileManager)
 		: LauncherProfileManager(ProfileManager)
 	{ 
-		SetDefaults();
 		ProjectChangedDelegate.AddRaw(this, &FLauncherProfile::OnSelectedProjectChanged);
 	}
 
@@ -214,7 +213,6 @@ public:
 		, Id(InId)
 		, Name(InProfileName)
 	{
-		SetDefaults();
 		ProjectChangedDelegate.AddRaw(this, &FLauncherProfile::OnSelectedProjectChanged);
 	}
 
diff --git a/Engine/Source/Developer/LauncherServices/Private/Profiles/LauncherProfileAutomatedTest.h b/Engine/Source/Developer/LauncherServices/Private/Profiles/LauncherProfileAutomatedTest.h
--- a/Engine/Source/Developer/LauncherServices/Private/Profiles/LauncherProfileAutomatedTest.h
+++ b/Engine/Source/Developer/LauncherServices/Private/Profiles/LauncherProfileAutomatedTest.h
@@ -14,7 +14,7 @@ class FLauncherProfileAutomatedTest final : public ILauncherProfileAutomatedTest
 public:
 	static const TCHAR* TypeName;
 
-	FLauncherProfileAutomatedTest( FLauncherProfile* InProfile, const TCHAR* InInternalName = TEXT(""), const TCHAR* InUserTypeName = TEXT("") ) 
+	FLauncherProfileAutomatedTest( TSharedPtr<FLauncherProfile> InProfile, const TCHAR* InInternalName = TEXT(""), const TCHAR* InUserTypeName = TEXT("") ) 
 		: Profile(InProfile)
 		, InternalName(InInternalName)
 		, UserTypeName(InUserTypeName)
@@ -218,7 +218,7 @@ public:
 
 
 private:
-	FLauncherProfile* Profile;
+	TSharedPtr<FLauncherProfile> Profile;
 
 	FString UATCommand = TEXT("RunUnreal");
 	FString Tests; //-Test=xxx,yyy
diff --git a/Engine/Source/Developer/LauncherServices/Private/Profiles/LauncherProfileBuildCookRun.h b/Engine/Source/Developer/LauncherServices/Private/Profiles/LauncherProfileBuildCookRun.h
--- a/Engine/Source/Developer/LauncherServices/Private/Profiles/LauncherProfileBuildCookRun.h
+++ b/Engine/Source/Developer/LauncherServices/Private/Profiles/LauncherProfileBuildCookRun.h
@@ -72,7 +72,7 @@ public:
 	/**
 	 * Default constructor.
 	 */
-	FLauncherProfileBuildCookRun( ILauncherProfileManagerRef InProfileManager, FLauncherProfile* InProfile, const TCHAR* InInternalName = TEXT(""), const TCHAR* InUserTypeName = TEXT("") ) 
+	FLauncherProfileBuildCookRun( ILauncherProfileManagerRef InProfileManager, TSharedPtr<FLauncherProfile> InProfile, const TCHAR* InInternalName = TEXT(""), const TCHAR* InUserTypeName = TEXT("") ) 
 		: LauncherProfileManager(InProfileManager)		
 		, DefaultLaunchRole(MakeShareable(new FLauncherProfileLaunchRole()))
 		, Profile(InProfile)
@@ -2102,7 +2102,7 @@ private:
 	TArray<FString> EditorArchitectures;
 
 private:
-	FLauncherProfile* Profile;
+	TSharedPtr<FLauncherProfile> Profile;
 
 	FString UATCommand = TEXT("BuildCookRun");
 	FString UATCommandLine;
diff --git a/Engine/Source/Developer/LauncherServices/Private/Profiles/LauncherProfileManager.cpp b/Engine/Source/Developer/LauncherServices/Private/Profiles/LauncherProfileManager.cpp
--- a/Engine/Source/Developer/LauncherServices/Private/Profiles/LauncherProfileManager.cpp
+++ b/Engine/Source/Developer/LauncherServices/Private/Profiles/LauncherProfileManager.cpp
@@ -155,7 +155,8 @@ ILauncherProfileRef FLauncherProfileManager::AddNewProfile()
 	}
 
 	// create and add the profile
-	ILauncherProfileRef NewProfile = MakeShareable(new FLauncherProfile(AsShared(), FGuid::NewGuid(), ProfileName));
+	TSharedRef<FLauncherProfile> NewProfile = MakeShared<FLauncherProfile>(AsShared(), FGuid::NewGuid(), ProfileName);
+	NewProfile->SetDefaults();
 
 	AddProfile(NewProfile);
 
@@ -167,7 +168,8 @@ ILauncherProfileRef FLauncherProfileManager::AddNewProfile()
 ILauncherProfileRef FLauncherProfileManager::CreateUnsavedProfile(FString ProfileName)
 {
 	// create and return the profile
-	ILauncherProfileRef NewProfile = MakeShareable(new FLauncherProfile(AsShared(), FGuid(), ProfileName));
+	TSharedRef<FLauncherProfile> NewProfile = MakeShared<FLauncherProfile>(AsShared(), FGuid(), ProfileName);
+	NewProfile->SetDefaults();
 	
 	AllProfiles.Add(NewProfile);
 	
@@ -266,7 +268,8 @@ ILauncherProfilePtr FLauncherProfileManager::GetProfile( const FGuid& ProfileId
 
 ILauncherProfilePtr FLauncherProfileManager::LoadProfile( FArchive& Archive )
 {
-	FLauncherProfile* Profile = new FLauncherProfile(AsShared());
+	TSharedRef<FLauncherProfile> Profile = MakeShared<FLauncherProfile>(AsShared());
+	Profile->SetDefaults();
 
 PRAGMA_DISABLE_DEPRECATION_WARNINGS
 	if (Profile->Serialize(Archive))
@@ -282,7 +285,7 @@ PRAGMA_ENABLE_DEPRECATION_WARNINGS
 			BuildCookRun->SetDeployedDeviceGroup(DeviceGroup);
 		}
 
-		return MakeShareable(Profile);
+		return Profile;
 	}
 
 	return nullptr;
@@ -303,7 +306,8 @@ ILauncherProfilePtr FLauncherProfileManager::LoadJSONProfile(FString ProfileFile
 		return nullptr;
 	}
 
-	FLauncherProfile* Profile = new FLauncherProfile(AsShared());
+	TSharedRef<FLauncherProfile> Profile = MakeShared<FLauncherProfile>(AsShared());
+	Profile->SetDefaults();
 
 	if (Profile->Load(*(Object.Get())))
 	{
@@ -317,7 +321,7 @@ ILauncherProfilePtr FLauncherProfileManager::LoadJSONProfile(FString ProfileFile
 			BuildCookRun->SetDeployedDeviceGroup(DeviceGroup);
 		}
 
-		return MakeShareable(Profile);
+		return Profile;
 	}
 
 	return nullptr;
diff --git a/Engine/Source/Developer/LauncherServices/Private/Profiles/LauncherProfileUATCommand.h b/Engine/Source/Developer/LauncherServices/Private/Profiles/LauncherProfileUATCommand.h
--- a/Engine/Source/Developer/LauncherServices/Private/Profiles/LauncherProfileUATCommand.h
+++ b/Engine/Source/Developer/LauncherServices/Private/Profiles/LauncherProfileUATCommand.h
@@ -15,7 +15,7 @@ class FLauncherProfileUATCommand final : public ILauncherProfileUATCommand, publ
 public:
 	static const TCHAR* TypeName;
 
-	FLauncherProfileUATCommand(FLauncherProfile* InProfile, const TCHAR* InInternalName = TEXT(""), const TCHAR* InUserTypeName = TEXT("") ) 
+	FLauncherProfileUATCommand(TSharedPtr<FLauncherProfile> InProfile, const TCHAR* InInternalName = TEXT(""), const TCHAR* InUserTypeName = TEXT("") ) 
 		: Profile(InProfile)
 		, InternalName(InInternalName)
 		, UserTypeName(InUserTypeName)
@@ -177,7 +177,7 @@ public:
 
 
 private:
-	FLauncherProfile* Profile;
+	TSharedPtr<FLauncherProfile> Profile;
 
 	FString UATCommand = TEXT("RunUnreal");
 	FString UATCommandLine;
