Cache and validate the build targets if a profile will require the extra validation
- likely profiles created with the legacy project launcher

#rnx

#Synced-CL 51007224

diff --git a/Engine/Source/Developer/LauncherServices/Private/Profiles/LauncherProfile.h b/Engine/Source/Developer/LauncherServices/Private/Profiles/LauncherProfile.h
--- a/Engine/Source/Developer/LauncherServices/Private/Profiles/LauncherProfile.h
+++ b/Engine/Source/Developer/LauncherServices/Private/Profiles/LauncherProfile.h
@@ -2801,8 +2801,9 @@ protected:
 
 	void ValidateBuildTarget(const ILauncherProfileBuildCookRunRef& BuildCookRun)
 	{
-		// this is only useful if the legacy project launcher is being used
-		if (!BuildTargetOptionsChangedDelegate.IsBound())
+		// this is only necessary if the legacy project launcher was used to create complex build target situations
+		// or if the legacy project launcher is being used	
+		if (BuildCookRun->GetCookedPlatforms().Num() == 1 && BuildCookRun->GetBuildTargets().Num() <= 1 && !BuildTargetOptionsChangedDelegate.IsBound() )
 		{
 			return;
 		}
@@ -2892,8 +2893,18 @@ PRAGMA_ENABLE_DEPRECATION_WARNINGS
 
 	void RefreshValidBuildTargets()
 	{
-		// this is only useful if the legacy project launcher is being used
-		if (!BuildTargetOptionsChangedDelegate.IsBound())
+		// this is only necessary if the legacy project launcher was used to create complex build target situations
+		// or if the legacy project launcher is being used	
+		bool bHasComplexBuildTargets = false;
+		for (const ILauncherProfileBuildCookRunRef& BuildCookRun : GetBuildCookRunCommands())
+		{
+			if (BuildCookRun->GetCookedPlatforms().Num() > 1 || BuildCookRun->GetBuildTargets().Num() > 1)
+			{
+				bHasComplexBuildTargets = true;
+				break;
+			}
+		}
+		if (!bHasComplexBuildTargets && !BuildTargetOptionsChangedDelegate.IsBound())
 		{
 			return;
 		}
