logic fix for profile validation early-out

#rnx

#Synced-CL 51007224

diff --git a/Engine/Source/Developer/LauncherServices/Private/Profiles/LauncherProfile.h b/Engine/Source/Developer/LauncherServices/Private/Profiles/LauncherProfile.h
--- a/Engine/Source/Developer/LauncherServices/Private/Profiles/LauncherProfile.h
+++ b/Engine/Source/Developer/LauncherServices/Private/Profiles/LauncherProfile.h
@@ -2803,7 +2803,7 @@ protected:
 	{
 		// this is only necessary if the legacy project launcher was used to create complex build target situations
 		// or if the legacy project launcher is being used	
-		if (BuildCookRun->GetCookedPlatforms().Num() == 1 && BuildCookRun->GetBuildTargets().Num() <= 1 && !BuildTargetOptionsChangedDelegate.IsBound() )
+		if (BuildCookRun->GetCookedPlatforms().Num() <= 1 && BuildCookRun->GetBuildTargets().Num() <= 1 && !BuildTargetOptionsChangedDelegate.IsBound() )
 		{
 			return;
 		}
