launch command line fixes:
- strip out duplicate launch parameters
- unified BuildCookRun launch parameters depend on target type
- unified BuildCookRun honors disabled items

#rnx

#Synced-CL 51094919
#Synced-CL 51094919

diff --git a/Engine/Source/Developer/LauncherServices/Private/Launcher/LauncherWorker.cpp b/Engine/Source/Developer/LauncherServices/Private/Launcher/LauncherWorker.cpp
--- a/Engine/Source/Developer/LauncherServices/Private/Launcher/LauncherWorker.cpp
+++ b/Engine/Source/Developer/LauncherServices/Private/Launcher/LauncherWorker.cpp
@@ -652,7 +652,7 @@ FString FLauncherWorker::CreateUATCommandForBuildCookRun( const ILauncherProfile
 			*EscapedInheritableGameOptions);
 		if (!AddCmdLine.TrimStartAndEnd().IsEmpty()) // don't bother with -addcmdline if it's empty
 		{
-			AdditionalCommandLine = FString::Printf(TEXT(" -addcmdline=\"%s\""), *AddCmdLine);
+			AdditionalCommandLine = FString::Printf(TEXT(" -addcmdline=\"%s\""), *FinalizeCommandLine(AddCmdLine));
 		}
 
 		for (EBuildTargetType BuildTarget : BuildTargets)
@@ -669,14 +669,14 @@ FString FLauncherWorker::CreateUATCommandForBuildCookRun( const ILauncherProfile
 			{
 				case EBuildTargetType::Server:
 				{
-					AdditionalCommandLine += FString::Printf(TEXT(" -servercmdline=\"%s\""), *BuildTargetCmdLine);
+					AdditionalCommandLine += FString::Printf(TEXT(" -servercmdline=\"%s\""), *FinalizeCommandLine(BuildTargetCmdLine));
 				}
 				break;
 
 				case EBuildTargetType::Client:
 				case EBuildTargetType::Game:
 				{
-					AdditionalCommandLine += FString::Printf(TEXT(" -clientcmdline=\"%s\""), *BuildTargetCmdLine);
+					AdditionalCommandLine += FString::Printf(TEXT(" -clientcmdline=\"%s\""), *FinalizeCommandLine(BuildTargetCmdLine));
 				}
 				break;
 			}
@@ -1207,7 +1207,7 @@ FString FLauncherWorker::CreateUATCommandForBuildCookRun( const ILauncherProfile
 
 
 
-	return UATCommand;
+	return FinalizeCommandLine(UATCommand);
 }
 
 
@@ -1376,7 +1376,7 @@ FString FLauncherWorker::CreateUATCommandForAutomatedTest( const ILauncherProfil
 		AutomatedTestCommand += AdditionalCommandLine;
 	}
 
-	return AutomatedTestCommand;
+	return FinalizeCommandLine(AutomatedTestCommand);
 }
 
 
@@ -1420,7 +1420,7 @@ FString FLauncherWorker::CreateUATCommand( const ILauncherProfileUATCommandRef&
 
 	FUATCommandPostProcessParameters PostProcessParameters =
 	{
-		.bShouldBeSkipped = true, // skip general-purpose UAT commands unless something explicitly requetss it
+		.bShouldBeSkipped = true, // skip general-purpose UAT commands unless something explicitly requets it
 		.Profile = InProfile,
 		.UATCommand = InUATCommand,
 		.UATCommandLine = UATCommandLine,
@@ -1435,7 +1435,7 @@ FString FLauncherWorker::CreateUATCommand( const ILauncherProfileUATCommandRef&
 	}
 
 
-	return UATCommandLine;
+	return FinalizeCommandLine(UATCommandLine);
 
 }
 
@@ -1501,7 +1501,7 @@ ILauncherProfileUATCommandPtr FLauncherWorker::MakeTemporaryUnifiedRunCommand(co
 	TArray<ILauncherProfileBuildCookRunRef> BuildCookRuns;
 	for (const ILauncherProfileBuildCookRunRef& BuildCookRun : InProfile->GetBuildCookRunCommands())
 	{
-		if (BuildCookRun->GetLaunchMode() != ELauncherProfileLaunchModes::DoNotLaunch && !BuildCookRun->IsUsingZenPakStreaming())
+		if (BuildCookRun->GetLaunchMode() != ELauncherProfileLaunchModes::DoNotLaunch && !BuildCookRun->IsUsingZenPakStreaming() && BuildCookRun->IsEnabled())
 		{
 			BuildCookRuns.Add(BuildCookRun);
 		}
@@ -1513,6 +1513,11 @@ ILauncherProfileUATCommandPtr FLauncherWorker::MakeTemporaryUnifiedRunCommand(co
 		return nullptr;
 	}
 
+	const TArray<FTargetInfo>& BuildTargets = FDesktopPlatformModule::Get()->GetTargetsForProject(InProfile->GetProjectPath());
+	ILauncherProfileManagerRef LauncherProfileManager = FModuleManager::LoadModuleChecked<ILauncherServicesModule>("LauncherServices").GetProfileManager();
+	ILauncherProfileBuildCookRunPtr UnifiedRunCommand = MakeShared<FLauncherProfileBuildCookRun>(LauncherProfileManager, nullptr, TEXT("TEMP_UnifiedRunCommand") );
+
+
 	auto CombineCommandLine = []( const FString& A, const FString& B)
 	{
 		if (A.IsEmpty())
@@ -1529,13 +1534,52 @@ ILauncherProfileUATCommandPtr FLauncherWorker::MakeTemporaryUnifiedRunCommand(co
 		}
 	};
 
+	auto GetCommandLine = []( const ILauncherProfileBuildCookRunPtr& BuildCookRun, EBuildTargetType BuildTarget )
+	{
+		if (BuildTarget == EBuildTargetType::Unknown)
+		{
+			return BuildCookRun->GetAdditionalCommandLineParameters();
+		}
+		else
+		{
+			return BuildCookRun->GetAdditionalTargetCommandLineParameters(BuildTarget);
+		}
+	};
 
-	ILauncherProfileManagerRef LauncherProfileManager = FModuleManager::LoadModuleChecked<ILauncherServicesModule>("LauncherServices").GetProfileManager();
-	ILauncherProfileBuildCookRunPtr UnifiedRunCommand = MakeShared<FLauncherProfileBuildCookRun>(LauncherProfileManager, nullptr, TEXT("TEMP_UnifiedRunCommand") );
+	auto SetCommandLine = []( const ILauncherProfileBuildCookRunPtr& BuildCookRun, EBuildTargetType BuildTarget, const FString& CommandLine )
+	{
+		if (BuildTarget == EBuildTargetType::Unknown)
+		{
+			BuildCookRun->SetAdditionalCommandLineParameters(CommandLine);
+		}
+		else
+		{
+			BuildCookRun->SetAdditionalTargetCommandLineParameters(CommandLine, BuildTarget);
+		}
+	};
 
-	UnifiedRunCommand->SetDeployedDeviceGroup( LauncherProfileManager->CreateUnmanagedDeviceGroup() );
+	auto MergeCommandLine = [this, UnifiedRunCommand, CombineCommandLine, GetCommandLine, SetCommandLine]( EBuildTargetType BuildTarget, const FString& OtherCommandLine)
+	{
+		FString CombinedCommandLine = FinalizeCommandLine(CombineCommandLine( GetCommandLine(UnifiedRunCommand, BuildTarget), OtherCommandLine ));
+		SetCommandLine( UnifiedRunCommand, BuildTarget, CombinedCommandLine);
+	};
+
+	auto GetBuildTargetType = [BuildTargets]( const FString& BuildTargetName )
+	{
+		for (const FTargetInfo& BuildTarget : BuildTargets)
+		{
+			if (BuildTarget.Name == BuildTargetName)
+			{
+				return BuildTarget.Type;
+			}
+		}
+
+		return EBuildTargetType::Unknown;
+	};
 
 
+	UnifiedRunCommand->SetDeployedDeviceGroup( LauncherProfileManager->CreateUnmanagedDeviceGroup() );
+
 	UnifiedRunCommand->SetDescription(TEXT("Unified Run"));
 	UnifiedRunCommand->SetBuildMode(ELauncherProfileBuildModes::DoNotBuild);
 	UnifiedRunCommand->SetCookMode(ELauncherProfileCookModes::DoNotCook);
@@ -1581,11 +1625,24 @@ ILauncherProfileUATCommandPtr FLauncherWorker::MakeTemporaryUnifiedRunCommand(co
 		UnifiedRunCommand->SetEditorArchitectures(EditorArchitectures);
 
 		// unify launch command lines
-		UnifiedRunCommand->SetAdditionalCommandLineParameters( CombineCommandLine( UnifiedRunCommand->GetAdditionalCommandLineParameters(), BuildCookRun->GetAdditionalCommandLineParameters() ) );
-		UnifiedRunCommand->SetAdditionalTargetCommandLineParameters( CombineCommandLine( UnifiedRunCommand->GetAdditionalTargetCommandLineParameters(EBuildTargetType::Client), BuildCookRun->GetAdditionalTargetCommandLineParameters(EBuildTargetType::Client) ), EBuildTargetType::Client );
-		UnifiedRunCommand->SetAdditionalTargetCommandLineParameters( CombineCommandLine( UnifiedRunCommand->GetAdditionalTargetCommandLineParameters(EBuildTargetType::Server), BuildCookRun->GetAdditionalTargetCommandLineParameters(EBuildTargetType::Server) ), EBuildTargetType::Server );
-		UnifiedRunCommand->SetAdditionalTargetCommandLineParameters( CombineCommandLine( UnifiedRunCommand->GetAdditionalTargetCommandLineParameters(EBuildTargetType::Game),   BuildCookRun->GetAdditionalTargetCommandLineParameters(EBuildTargetType::Game) ),   EBuildTargetType::Game   );
-		UnifiedRunCommand->SetAdditionalTargetCommandLineParameters( CombineCommandLine( UnifiedRunCommand->GetAdditionalTargetCommandLineParameters(EBuildTargetType::Editor), BuildCookRun->GetAdditionalTargetCommandLineParameters(EBuildTargetType::Editor) ), EBuildTargetType::Editor );
+		EBuildTargetType SingleBuildTarget = (BuildCookRun->GetBuildTargets().Num() == 1) ? GetBuildTargetType(BuildCookRun->GetBuildTargets()[0]) : EBuildTargetType::Unknown;
+
+		const EBuildTargetType SupportedBuildTargets[] = { EBuildTargetType::Unknown, EBuildTargetType::Client, EBuildTargetType::Server, EBuildTargetType::Game, EBuildTargetType::Editor };
+		for (const EBuildTargetType BuildTarget : SupportedBuildTargets)
+		{
+			// copy the per-target type command line over - only if the build target matches, or the source has many targets anyway
+			if (SingleBuildTarget == EBuildTargetType::Unknown || SingleBuildTarget == BuildTarget)
+			{
+				MergeCommandLine( BuildTarget, GetCommandLine(BuildCookRun, BuildTarget) );
+			}
+
+			// copy the target-agnostic command line over into the target-specific command line property
+			if (SingleBuildTarget != EBuildTargetType::Unknown && BuildTarget == EBuildTargetType::Unknown)
+			{
+				MergeCommandLine( SingleBuildTarget, GetCommandLine(BuildCookRun, BuildTarget) );
+			}
+
+		}
 
 		// unify devices
 		for ( const FString& DeviceId : BuildCookRun->GetDeployedDeviceGroup()->GetDeviceIDs())
@@ -1644,6 +1701,21 @@ TArray<FString> FLauncherWorker::FinalizePlatforms( const TArray<FString>& CookP
 	return MoveTemp(Platforms);
 }
 
+FString FLauncherWorker::FinalizeCommandLine( const FString& CommandLine ) const
+{
+	const TCHAR* CommandLinePtr = CommandLine.GetCharArray().GetData();
+	
+	TArray<FString> Result;
+
+	FString Token;
+	while( FParse::Token(CommandLinePtr, Token, false) )
+	{
+		Result.AddUnique(Token);
+	}
+
+	return FString::Join(Result, TEXT(" "));
+}
+
 
 
 /* FLauncherWorker implementation
diff --git a/Engine/Source/Developer/LauncherServices/Private/Launcher/LauncherWorker.h b/Engine/Source/Developer/LauncherServices/Private/Launcher/LauncherWorker.h
--- a/Engine/Source/Developer/LauncherServices/Private/Launcher/LauncherWorker.h
+++ b/Engine/Source/Developer/LauncherServices/Private/Launcher/LauncherWorker.h
@@ -129,6 +129,7 @@ protected:
 
 	ILauncherProfileUATCommandPtr MakeTemporaryUnifiedRunCommand(const ILauncherProfileRef& InProfile);
 	TArray<FString> FinalizePlatforms( const TArray<FString>& CookPlatforms, const ILauncherDeviceGroupPtr& DeviceGroup, bool bPreferCookedPlatforms ) const;
+	FString FinalizeCommandLine( const FString& CommandLine ) const;
 
 	//start listening to the device discovery routine
 	void EnableDeviceDiscoveryListener(ILauncherProfileUATCommandPtr UnifiedUATCommand);
