fix for command line finalization

#rnx

#Synced-CL 51094919

diff --git a/Engine/Source/Developer/LauncherServices/Private/Launcher/LauncherWorker.cpp b/Engine/Source/Developer/LauncherServices/Private/Launcher/LauncherWorker.cpp
--- a/Engine/Source/Developer/LauncherServices/Private/Launcher/LauncherWorker.cpp
+++ b/Engine/Source/Developer/LauncherServices/Private/Launcher/LauncherWorker.cpp
@@ -1708,7 +1708,7 @@ FString FLauncherWorker::FinalizeCommandLine( const FString& CommandLine ) const
 	TArray<FString> Result;
 
 	FString Token;
-	while( FParse::Token(CommandLinePtr, Token, false) )
+	while( CommandLinePtr != nullptr && FParse::Token(CommandLinePtr, Token, false) )
 	{
 		Result.AddUnique(Token);
 	}
