don't add No-Op BuildCookRun commands to the UAT command line

#rnx

#Synced-CL 51094919

diff --git a/Engine/Source/Developer/LauncherServices/Private/Launcher/LauncherWorker.cpp b/Engine/Source/Developer/LauncherServices/Private/Launcher/LauncherWorker.cpp
--- a/Engine/Source/Developer/LauncherServices/Private/Launcher/LauncherWorker.cpp
+++ b/Engine/Source/Developer/LauncherServices/Private/Launcher/LauncherWorker.cpp
@@ -385,7 +385,19 @@ FString FLauncherWorker::CreateUATCommandForBuildCookRun( const ILauncherProfile
 
 	CommandStart = TEXT("");
 
-	
+	// early out if this will be a no-op
+	if (!InBuildCookRun->IsImportingZenSnapshot() &&
+		!InBuildCookRun->IsArchiving() &&
+		InBuildCookRun->GetBuildMode() == ELauncherProfileBuildModes::DoNotBuild &&
+		InBuildCookRun->GetCookMode() == ELauncherProfileCookModes::DoNotCook &&
+		InBuildCookRun->GetDeploymentMode() == ELauncherProfileDeploymentModes::DoNotDeploy &&
+		(InBuildCookRun->GetLaunchMode() == ELauncherProfileLaunchModes::DoNotLaunch || Adjustments.bDisableRun))
+	{
+		UE_LOG(LogLauncherProfile, Log, TEXT("Skipping BuildCookRun \"%s\" because it will not do anything"), *InBuildCookRun->GetDescription());
+		return FString();
+	}
+
+
 	bool bPreferCookedPlatforms = (Adjustments.bPreferCookedPlatforms || InBuildCookRun->GetCookMode() == ELauncherProfileCookModes::ByTheBook || InBuildCookRun->ShouldBuild());
 	TArray<FString> FinalPlatforms = FinalizePlatforms(InBuildCookRun->GetCookedPlatforms(), InBuildCookRun->GetDeployedDeviceGroup(), bPreferCookedPlatforms);
 
