Project Launcher 2 - move zen streaming options to its own version

#Synced-CL 51262173

diff --git a/Engine/Source/Developer/LauncherServices/Private/Profiles/LauncherProfileBuildCookRun.h b/Engine/Source/Developer/LauncherServices/Private/Profiles/LauncherProfileBuildCookRun.h
--- a/Engine/Source/Developer/LauncherServices/Private/Profiles/LauncherProfileBuildCookRun.h
+++ b/Engine/Source/Developer/LauncherServices/Private/Profiles/LauncherProfileBuildCookRun.h
@@ -1073,10 +1073,14 @@ public:
 			bUseZenStore = Object.GetBoolField(TEXT("UseZenStore"));
 			bImportingZenSnapshot = Object.GetBoolField(TEXT("ImportingZenSnapshot"));
 			bUseZenPakStreaming = Object.GetBoolField(TEXT("UseZenPakStreaming"));
-			bUseZenStreaming = Object.GetBoolField(TEXT("UseZenStreaming"));
 			ZenPakStreamingPath = Object.GetStringField(TEXT("ZenPakStreamingPath"));
 		}
 
+		if (ProfileVersion >= LAUNCHERSERVICES_ADDZENSTREAMINGOPTIONS)
+		{
+			bUseZenStreaming = Object.GetBoolField(TEXT("UseZenStreaming"));
+		}
+
 		if (ProfileVersion >= LAUNCHERSERVICES_ADDBUILDARCHITECTURE)
 		{
 			ClientArchitectures.Reset();
diff --git a/Engine/Source/Developer/LauncherServices/Private/Profiles/LauncherProfileVersion.h b/Engine/Source/Developer/LauncherServices/Private/Profiles/LauncherProfileVersion.h
--- a/Engine/Source/Developer/LauncherServices/Private/Profiles/LauncherProfileVersion.h
+++ b/Engine/Source/Developer/LauncherServices/Private/Profiles/LauncherProfileVersion.h
@@ -43,6 +43,7 @@ enum ELauncherVersion
 	LAUNCHERSERVICES_ADDMULTIPLEBUILDTARGETS = 44,
 	LAUNCHERSERVICES_ADDUATCOMMANDS = 45,
 	LAUNCHERSERVICES_INNERBUILDCOOKRUN = 46,
+	LAUNCHERSERVICES_ADDZENSTREAMINGOPTIONS = 47,
 	//ADD NEW STUFF HERE
 
 
