Added PreStaged option to BuildCookRun parameters to allow the launching of externally-produced staged builds.

#rnx

#Synced-CL 51338024
#Synced-CL 51338024

diff --git a/Engine/Plugins/Developer/ProjectLauncher/Source/ProjectLauncher/Private/Model/ProjectLauncherModel.cpp b/Engine/Plugins/Developer/ProjectLauncher/Source/ProjectLauncher/Private/Model/ProjectLauncherModel.cpp
--- a/Engine/Plugins/Developer/ProjectLauncher/Source/ProjectLauncher/Private/Model/ProjectLauncherModel.cpp
+++ b/Engine/Plugins/Developer/ProjectLauncher/Source/ProjectLauncher/Private/Model/ProjectLauncherModel.cpp
@@ -839,6 +839,10 @@ PRAGMA_ENABLE_DEPRECATION_WARNINGS
 		{
 			return EContentScheme::CookOnTheFly;
 		}
+		else if (BuildCookRun->IsUsingPreStagedBuild())
+		{
+			return EContentScheme::PreStagedBuild;
+		}
 		else if (BuildCookRun->GetPackagingMode() != ELauncherProfilePackagingModes::DoNotPackage)
 		{
 			if (BuildCookRun->IsCreatingReleaseVersion() || BuildCookRun->IsGeneratingPatch() || BuildCookRun->IsCreatingDLC())
@@ -879,9 +883,12 @@ PRAGMA_ENABLE_DEPRECATION_WARNINGS
 		bool bPackage = (ContentScheme == EContentScheme::DevelopmentPackage || ContentScheme == EContentScheme::SubmissionPackage);
 		bool bZenPakStreaming = (ContentScheme == EContentScheme::ZenPakStreaming);
 		bool bSubmissionPackage = (ContentScheme == EContentScheme::SubmissionPackage);
+		bool bPreStagedBuild = (ContentScheme == EContentScheme::PreStagedBuild);
+
 		// FIXME: may need to turn off Zen/Pak in UAT somehow if it's set in project defaults?
 
 		BuildCookRun->SetUseZenStreaming(bUseZenStreaming);
+		BuildCookRun->SetUsePreStagedBuild(bPreStagedBuild);
 		BuildCookRun->SetUseZenPakStreaming(bZenPakStreaming);
 		BuildCookRun->SetDeployWithUnrealPak(bPakFiles);
 		BuildCookRun->SetUseZenStore(bUseZen);
@@ -1389,6 +1396,7 @@ PRAGMA_ENABLE_DEPRECATION_WARNINGS
 			case EContentScheme::SubmissionPackage:  return LOCTEXT("ContentSchemeSubPackage","Submission Package");
 			case EContentScheme::LooseFiles:         return LOCTEXT("ContentSchemeLooseFiles","Loose Files (legacy)");
 			case EContentScheme::CookOnTheFly:       return LOCTEXT("ContentSchemeCOTF","Cook On The Fly"); 
+			case EContentScheme::PreStagedBuild:     return LOCTEXT("ContentSchemePreStaged", "Use Current Staged Build");
 		}
 
 		checkNoEntry();
@@ -1407,6 +1415,7 @@ PRAGMA_ENABLE_DEPRECATION_WARNINGS
 			case EContentScheme::SubmissionPackage:  return LOCTEXT("ContentSchemeTipSubPackage","Package cooked game content into a single package or patch for submission to a game store, where available");
 			case EContentScheme::LooseFiles:         return LOCTEXT("ContentSchemeTipLooseFiles","Store cooked game assets in individual files (legacy - recommend moving to Zen Streaming. This option will not work if the project is already configured with 'Use Zen Store')");
 			case EContentScheme::CookOnTheFly:       return LOCTEXT("ContentSchemeTipCOTF","Only cook game assets when the game requires them, and send them over the network (legacy - slow)");
+			case EContentScheme::PreStagedBuild:     return LOCTEXT("ContentSchemeTipPreStaged","Use the existing staged build as-is. Works best with the Build Sync extension or Build Storage Tool etc");
 		}
 
 		checkNoEntry();
@@ -1428,6 +1437,7 @@ PRAGMA_ENABLE_DEPRECATION_WARNINGS
 			case ProjectLauncher::EContentScheme::SubmissionPackage:  return TEXT("SubmissionPackage");
 			case ProjectLauncher::EContentScheme::LooseFiles:         return TEXT("LooseFiles");
 			case ProjectLauncher::EContentScheme::CookOnTheFly:       return TEXT("CookOnTheFly");
+			case ProjectLauncher::EContentScheme::PreStagedBuild:     return TEXT("PreStagedBuild");
 			default: return TEXT("Unknown");
 		}
 	}
@@ -1471,6 +1481,11 @@ PRAGMA_ENABLE_DEPRECATION_WARNINGS
 			OutContentScheme = ProjectLauncher::EContentScheme::CookOnTheFly;
 			return true;
 		}
+		else if (FCString::Stricmp(String, TEXT("PreStagedBuild")) == 0) 
+		{
+			OutContentScheme = ProjectLauncher::EContentScheme::PreStagedBuild;
+			return true;
+		}
 		else
 		{
 			return false;
diff --git a/Engine/Plugins/Developer/ProjectLauncher/Source/ProjectLauncher/Private/ProfileTree/CustomProfileTreeBuilder.cpp b/Engine/Plugins/Developer/ProjectLauncher/Source/ProjectLauncher/Private/ProfileTree/CustomProfileTreeBuilder.cpp
--- a/Engine/Plugins/Developer/ProjectLauncher/Source/ProjectLauncher/Private/ProfileTree/CustomProfileTreeBuilder.cpp
+++ b/Engine/Plugins/Developer/ProjectLauncher/Source/ProjectLauncher/Private/ProfileTree/CustomProfileTreeBuilder.cpp
@@ -52,6 +52,7 @@ namespace ProjectLauncher
 
 			BuildCookRun.AddTargetProperty(GeneralBuildCookRunHeader);
 			BuildCookRun.AddPlatformProperty(GeneralBuildCookRunHeader);
+			BuildCookRun.AddArchitectureProperty(GeneralBuildCookRunHeader);
 			BuildCookRun.AddConfigurationProperty(GeneralBuildCookRunHeader);
 			BuildCookRun.AddContentSchemeProperty(GeneralBuildCookRunHeader);
 
@@ -71,7 +72,6 @@ namespace ProjectLauncher
 			FLaunchProfileTreeNode& BuildHeader = AddHeading( GeneralBuildCookRunHeader,  TEXT("Build"), LOCTEXT("BuildHeading", "Build") );
 			BuildCookRun.AddBuildProperty(BuildHeader);
 			BuildCookRun.AddForceBuildProperty(BuildHeader);
-			BuildCookRun.AddArchitectureProperty(BuildHeader);
 
 			FLaunchProfileTreeNode& DirectoryHeader = AddHeading( GeneralBuildCookRunHeader,  TEXT("Directory"), LOCTEXT("DirectoryHeading", "Directory") );
 			BuildCookRun.AddArchiveBuildProperty(DirectoryHeader);
diff --git a/Engine/Plugins/Developer/ProjectLauncher/Source/ProjectLauncher/Private/ProfileTree/GenericProfileTreeBuilder.cpp b/Engine/Plugins/Developer/ProjectLauncher/Source/ProjectLauncher/Private/ProfileTree/GenericProfileTreeBuilder.cpp
--- a/Engine/Plugins/Developer/ProjectLauncher/Source/ProjectLauncher/Private/ProfileTree/GenericProfileTreeBuilder.cpp
+++ b/Engine/Plugins/Developer/ProjectLauncher/Source/ProjectLauncher/Private/ProfileTree/GenericProfileTreeBuilder.cpp
@@ -80,17 +80,17 @@ namespace ProjectLauncher
 
 		ForCooked = [this]()
 		{
-			return (ContentScheme != EContentScheme::ZenPakStreaming && ContentScheme != EContentScheme::CookOnTheFly);
+			return (ContentScheme != EContentScheme::ZenPakStreaming && ContentScheme != EContentScheme::CookOnTheFly && ContentScheme != EContentScheme::PreStagedBuild);
 		};
 
 		ForEnabledCooked = [this]()
 		{
-			return (bShouldCook && ContentScheme != EContentScheme::ZenPakStreaming && ContentScheme != EContentScheme::CookOnTheFly);
+			return (bShouldCook && ContentScheme != EContentScheme::ZenPakStreaming && ContentScheme != EContentScheme::CookOnTheFly && ContentScheme != EContentScheme::PreStagedBuild);
 		};
 
 		ForContent = [this]()
 		{
-			return (ContentScheme != EContentScheme::ZenPakStreaming);
+			return (ContentScheme != EContentScheme::ZenPakStreaming && ContentScheme != EContentScheme::PreStagedBuild);
 		};
 
 		ForCode = [this]()
@@ -98,6 +98,11 @@ namespace ProjectLauncher
 			return GetBuild();
 		};
 
+		ForCodeBuild = [this]()
+		{
+			return (ContentScheme != EContentScheme::PreStagedBuild);
+		};
+
 		ForDeployment = [this]()
 		{
 			return GetDeployToDevice();
@@ -584,6 +589,7 @@ namespace ProjectLauncher
 				.GetValue = [this]()			{ return GetBuild(); },
 				.SetValue = [this](bool bValue)	{ SetBuild(bValue); },
 				.GetDefaultValue = [this]()		{ return GetBuild(Owner.DefaultBuildCookRun); },
+				.IsVisible = ForCodeBuild,
 			}
 		);
 	}
@@ -595,6 +601,7 @@ namespace ProjectLauncher
 				.GetValue = [this]()			{ return GetForceBuild(); },
 				.SetValue = [this](bool bValue)	{ SetForceBuild(bValue); },
 				.GetDefaultValue = [this]()		{ return GetForceBuild(Owner.DefaultBuildCookRun); },
+				.IsVisible = ForCodeBuild,
 				.IsEnabled = ForCode,
 			}
 		);
diff --git a/Engine/Plugins/Developer/ProjectLauncher/Source/ProjectLauncher/Public/Model/ProjectLauncherModel.h b/Engine/Plugins/Developer/ProjectLauncher/Source/ProjectLauncher/Public/Model/ProjectLauncherModel.h
--- a/Engine/Plugins/Developer/ProjectLauncher/Source/ProjectLauncher/Public/Model/ProjectLauncherModel.h
+++ b/Engine/Plugins/Developer/ProjectLauncher/Source/ProjectLauncher/Public/Model/ProjectLauncherModel.h
@@ -40,6 +40,7 @@ namespace ProjectLauncher
 		SubmissionPackage,
 		LooseFiles,
 		CookOnTheFly,
+		PreStagedBuild,
 		MAX
 	};
 
diff --git a/Engine/Plugins/Developer/ProjectLauncher/Source/ProjectLauncher/Public/ProfileTree/GenericProfileTreeBuilder.h b/Engine/Plugins/Developer/ProjectLauncher/Source/ProjectLauncher/Public/ProfileTree/GenericProfileTreeBuilder.h
--- a/Engine/Plugins/Developer/ProjectLauncher/Source/ProjectLauncher/Public/ProfileTree/GenericProfileTreeBuilder.h
+++ b/Engine/Plugins/Developer/ProjectLauncher/Source/ProjectLauncher/Public/ProfileTree/GenericProfileTreeBuilder.h
@@ -164,6 +164,7 @@ namespace ProjectLauncher
 			FLaunchProfileTreeNode::FGetBool ForEnabledCooked;
 			FLaunchProfileTreeNode::FGetBool ForContent;
 			FLaunchProfileTreeNode::FGetBool ForCode;
+			FLaunchProfileTreeNode::FGetBool ForCodeBuild;
 			FLaunchProfileTreeNode::FGetBool ForDeployment;
 			FLaunchProfileTreeNode::FGetBool ForRun;
 			FLaunchProfileTreeNode::FGetBool ForSubmissionPackage;
diff --git a/Engine/Source/Developer/LauncherServices/Private/Launcher/LauncherWorker.cpp b/Engine/Source/Developer/LauncherServices/Private/Launcher/LauncherWorker.cpp
--- a/Engine/Source/Developer/LauncherServices/Private/Launcher/LauncherWorker.cpp
+++ b/Engine/Source/Developer/LauncherServices/Private/Launcher/LauncherWorker.cpp
@@ -697,6 +697,8 @@ FString FLauncherWorker::CreateUATCommandForBuildCookRun( const ILauncherProfile
 		UATCommand += AdditionalCommandLine;
 	}
 
+	bool bPreStaged = InBuildCookRun->IsUsingPreStagedBuild();
+
 
 	// map list
 	FString MapList;
@@ -733,7 +735,7 @@ FString FLauncherWorker::CreateUATCommandForBuildCookRun( const ILauncherProfile
 	bool bIsBuilding = InBuildCookRun->ShouldBuild();
 
 	// build
-	if (bIsBuilding)
+	if (bIsBuilding && !bPreStaged)
 	{
 		UATCommand += TEXT(" -build");
 
@@ -753,7 +755,7 @@ FString FLauncherWorker::CreateUATCommandForBuildCookRun( const ILauncherProfile
 	}
 
 	// snapshot import
-	if (InBuildCookRun->IsImportingZenSnapshot())
+	if (InBuildCookRun->IsImportingZenSnapshot() && !bPreStaged)
 	{
 		UATCommand += TEXT(" -snapshot");
 
@@ -767,302 +769,305 @@ FString FLauncherWorker::CreateUATCommandForBuildCookRun( const ILauncherProfile
 	}
 
 	// cook
-	switch(InBuildCookRun->GetCookMode())
+	if (!bPreStaged)
 	{
-	case ELauncherProfileCookModes::ByTheBook:
+		switch(InBuildCookRun->GetCookMode())
 		{
-			UATCommand += TEXT(" -cook");
-
-			UATCommand += MapList;
-			UATCommand += CultureList;
-
-			if (InBuildCookRun->IsCookingUnversioned())
+		case ELauncherProfileCookModes::ByTheBook:
 			{
-				UATCommand += TEXT(" -unversionedcookedcontent");
-			}
+				UATCommand += TEXT(" -cook");
 
-			if (InBuildCookRun->IsEncryptingIniFiles())
-			{
-				UATCommand += TEXT(" -encryptinifiles");
-			}
+				UATCommand += MapList;
+				UATCommand += CultureList;
 
-			TStringBuilder<64> InheritableCookOptions;
-			FCommandLine::BuildSubprocessCommandLine(ECommandLineArgumentFlags::CommandletContext, bOnlyInherited, InheritableCookOptions);
-			FString AdditionalOptions = InBuildCookRun->GetCookOptions();
-			if (!AdditionalOptions.IsEmpty() || (InheritableCookOptions.Len() > 0))
-			{
-				UATCommand += TEXT(" -additionalcookeroptions=\"");
+				if (InBuildCookRun->IsCookingUnversioned())
+				{
+					UATCommand += TEXT(" -unversionedcookedcontent");
+				}
 
-				// Escape any quotes in the argument list
-				UATCommand += AdditionalOptions.Replace(TEXT("\""), TEXT("\\\""));
+				if (InBuildCookRun->IsEncryptingIniFiles())
+				{
+					UATCommand += TEXT(" -encryptinifiles");
+				}
 
-				if (InheritableCookOptions.Len() > 0)
+				TStringBuilder<64> InheritableCookOptions;
+				FCommandLine::BuildSubprocessCommandLine(ECommandLineArgumentFlags::CommandletContext, bOnlyInherited, InheritableCookOptions);
+				FString AdditionalOptions = InBuildCookRun->GetCookOptions();
+				if (!AdditionalOptions.IsEmpty() || (InheritableCookOptions.Len() > 0))
 				{
-					if (!AdditionalOptions.IsEmpty())
+					UATCommand += TEXT(" -additionalcookeroptions=\"");
+
+					// Escape any quotes in the argument list
+					UATCommand += AdditionalOptions.Replace(TEXT("\""), TEXT("\\\""));
+
+					if (InheritableCookOptions.Len() > 0)
+					{
+						if (!AdditionalOptions.IsEmpty())
+						{
+							UATCommand += TEXT(" ");
+						}
+						FString EscapedInheritableCookOptions(InheritableCookOptions);
+						EscapedInheritableCookOptions.ReplaceInline(TEXT("\""), TEXT("\\\""));
+						UATCommand += EscapedInheritableCookOptions;
+					}
+
+					// If the additional options ends with a slash, make sure we don't escape the quote
+					if (UATCommand.EndsWith("\\"))
 					{
-						UATCommand += TEXT(" ");
+						UATCommand += TEXT("\\");
 					}
-					FString EscapedInheritableCookOptions(InheritableCookOptions);
-					EscapedInheritableCookOptions.ReplaceInline(TEXT("\""), TEXT("\\\""));
-					UATCommand += EscapedInheritableCookOptions;
+
+					UATCommand += TEXT("\"");
 				}
 
-				// If the additional options ends with a slash, make sure we don't escape the quote
-				if (UATCommand.EndsWith("\\"))
+				if (FParse::Param(FCommandLine::Get(), TEXT("fastcook")))
 				{
-					UATCommand += TEXT("\\");
+					// if our editor has nomcp then pass it through the launched game
+					UATCommand += TEXT(" -fastcook");
 				}
 
-				UATCommand += TEXT("\"");
-			}
+				if (InBuildCookRun->IsUsingZenStore())
+				{
+					UATCommand += TEXT(" -zenstore");
+				}
 
-			if (FParse::Param(FCommandLine::Get(), TEXT("fastcook")))
-			{
-				// if our editor has nomcp then pass it through the launched game
-				UATCommand += TEXT(" -fastcook");
-			}
+				if (FDerivedDataCacheInterface* DDC = TryGetDerivedDataCache())
+				{
+					const TCHAR* GraphName = DDC->GetGraphName();
+					if (FCString::Strcmp(GraphName, DDC->GetDefaultGraphName()))
+					{
+						UATCommand += FString::Printf(TEXT(" -DDC=%s"), GraphName);
+					}
+				}
 
-			if (InBuildCookRun->IsUsingZenStore())
-			{
-				UATCommand += TEXT(" -zenstore");
-			}
+				if (InBuildCookRun->IsPackingWithUnrealPak())
+				{
+					UATCommand += TEXT(" -pak");
+					UATCommand += TEXT(" -stage");
+					if (InBuildCookRun->IsUsingIoStore())
+					{
+						UATCommand += TEXT(" -iostore");
+					}
+					if (InBuildCookRun->IsCompressed())
+					{
+						UATCommand += TEXT(" -compressed");
+					}
+				}
 
-			if (FDerivedDataCacheInterface* DDC = TryGetDerivedDataCache())
-			{
-				const TCHAR* GraphName = DDC->GetGraphName();
-				if (FCString::Strcmp(GraphName, DDC->GetDefaultGraphName()))
+				if (InBuildCookRun->MakeBinaryConfig())
 				{
-					UATCommand += FString::Printf(TEXT(" -DDC=%s"), GraphName);
+					UATCommand += TEXT(" -makebinaryconfig");
 				}
-			}
 
-			if (InBuildCookRun->IsPackingWithUnrealPak())
-			{
-				UATCommand += TEXT(" -pak");
-				UATCommand += TEXT(" -stage");
-				if (InBuildCookRun->IsUsingIoStore())
+				if ( InBuildCookRun->IsCreatingReleaseVersion() )
 				{
-					UATCommand += TEXT(" -iostore");
+					UATCommand += TEXT(" -createreleaseversion=");
+					UATCommand += InBuildCookRun->GetCreateReleaseVersionName();
+				
 				}
-				if (InBuildCookRun->IsCompressed())
+
+				if ( InBuildCookRun->IsCreatingDLC() )
 				{
-					UATCommand += TEXT(" -compressed");
+					UATCommand += TEXT(" -dlcname=");
+					UATCommand += InBuildCookRun->GetDLCName();
 				}
-			}
 
-			if (InBuildCookRun->MakeBinaryConfig())
-			{
-				UATCommand += TEXT(" -makebinaryconfig");
-			}
+				if ( InBuildCookRun->IsDLCIncludingEngineContent() )
+				{
+					UATCommand += TEXT(" -DLCIncludeEngineContent");
+				}
 
-			if ( InBuildCookRun->IsCreatingReleaseVersion() )
-			{
-				UATCommand += TEXT(" -createreleaseversion=");
-				UATCommand += InBuildCookRun->GetCreateReleaseVersionName();
-				
-			}
+				if ( InBuildCookRun->IsGeneratingPatch() )
+				{
+					UATCommand += TEXT(" -generatepatch");
 
-			if ( InBuildCookRun->IsCreatingDLC() )
-			{
-				UATCommand += TEXT(" -dlcname=");
-				UATCommand += InBuildCookRun->GetDLCName();
-			}
+					if ( InBuildCookRun->ShouldAddPatchLevel() )
+					{
+						UATCommand += TEXT(" -addpatchlevel");
+					}
+				}
 
-			if ( InBuildCookRun->IsDLCIncludingEngineContent() )
-			{
-				UATCommand += TEXT(" -DLCIncludeEngineContent");
-			}
+				if ( InBuildCookRun->IsGeneratingPatch() || 
+					InBuildCookRun->IsCreatingReleaseVersion() || 
+					InBuildCookRun->IsCreatingDLC() )
+				{
+					if ( InBuildCookRun->GetBasedOnReleaseVersionName().IsEmpty() == false )
+					{
+						UATCommand += TEXT(" -basedonreleaseversion=");
+						UATCommand += InBuildCookRun->GetBasedOnReleaseVersionName();
+					}
 
-			if ( InBuildCookRun->IsGeneratingPatch() )
-			{
-				UATCommand += TEXT(" -generatepatch");
+					if (InBuildCookRun->GetOriginalReleaseVersionName().IsEmpty() == false)
+					{
+						UATCommand += TEXT(" -originalreleaseversion=");
+						UATCommand += InBuildCookRun->GetOriginalReleaseVersionName();
+					}
+				}
 
-				if ( InBuildCookRun->ShouldAddPatchLevel() )
+				if (InBuildCookRun->IsGeneratingChunks())
 				{
-					UATCommand += TEXT(" -addpatchlevel");
+					UATCommand += TEXT(" -manifests");
 				}
-			}
 
-			if ( InBuildCookRun->IsGeneratingPatch() || 
-				InBuildCookRun->IsCreatingReleaseVersion() || 
-				InBuildCookRun->IsCreatingDLC() )
-			{
-				if ( InBuildCookRun->GetBasedOnReleaseVersionName().IsEmpty() == false )
+				if (InBuildCookRun->IsGenerateHttpChunkData())
 				{
-					UATCommand += TEXT(" -basedonreleaseversion=");
-					UATCommand += InBuildCookRun->GetBasedOnReleaseVersionName();
+					auto Cmd = FString::Printf(TEXT(" -createchunkinstall -chunkinstalldirectory=\"%s\" -chunkinstallversion=\"%s\""), *InBuildCookRun->GetHttpChunkDataDirectory(), *InBuildCookRun->GetHttpChunkDataReleaseName());
+					UATCommand += Cmd;
 				}
-
-				if (InBuildCookRun->GetOriginalReleaseVersionName().IsEmpty() == false)
+			
+				// Creating a packed DLC requires staging
+				if (InBuildCookRun->GetPackagingMode() == ELauncherProfilePackagingModes::DoNotPackage && InBuildCookRun->IsCreatingDLC() && InBuildCookRun->IsPackingWithUnrealPak())
 				{
-					UATCommand += TEXT(" -originalreleaseversion=");
-					UATCommand += InBuildCookRun->GetOriginalReleaseVersionName();
+					UATCommand += TEXT(" -stage");
 				}
-			}
 
-			if (InBuildCookRun->IsGeneratingChunks())
-			{
-				UATCommand += TEXT(" -manifests");
+				FCommandDesc Desc;
+				FText Command = FText::Format(LOCTEXT("LauncherCookDesc", "Cook content for {0}"), FText::FromString(FString::Join(FinalPlatforms, TEXT("+") )));
+				Desc.Name = "Cook Task";
+				Desc.Desc = Command.ToString();
+				Desc.EndText = TEXT("********** COOK COMMAND COMPLETED **********");
+				OutCommands.Add(Desc);
+				if (CommandStart.Len() == 0)
+				{
+					CommandStart = TEXT("********** COOK COMMAND STARTED **********");
+				}
 			}
-
-			if (InBuildCookRun->IsGenerateHttpChunkData())
+			break;
+		case ELauncherProfileCookModes::OnTheFly:
 			{
-				auto Cmd = FString::Printf(TEXT(" -createchunkinstall -chunkinstalldirectory=\"%s\" -chunkinstallversion=\"%s\""), *InBuildCookRun->GetHttpChunkDataDirectory(), *InBuildCookRun->GetHttpChunkDataReleaseName());
-				UATCommand += Cmd;
-			}
+				UATCommand += TEXT(" -cookonthefly");
 			
-			// Creating a packed DLC requires staging
-			if (InBuildCookRun->GetPackagingMode() == ELauncherProfilePackagingModes::DoNotPackage && InBuildCookRun->IsCreatingDLC() && InBuildCookRun->IsPackingWithUnrealPak())
-			{
-				UATCommand += TEXT(" -stage");
-			}
+				if (InBuildCookRun->IsUsingZenStore())
+				{
+					UATCommand += TEXT(" -zenstore");
+				}
 
-			FCommandDesc Desc;
-			FText Command = FText::Format(LOCTEXT("LauncherCookDesc", "Cook content for {0}"), FText::FromString(FString::Join(FinalPlatforms, TEXT("+") )));
-			Desc.Name = "Cook Task";
-			Desc.Desc = Command.ToString();
-			Desc.EndText = TEXT("********** COOK COMMAND COMPLETED **********");
-			OutCommands.Add(Desc);
-			if (CommandStart.Len() == 0)
-			{
-				CommandStart = TEXT("********** COOK COMMAND STARTED **********");
-			}
-		}
-		break;
-	case ELauncherProfileCookModes::OnTheFly:
-		{
-			UATCommand += TEXT(" -cookonthefly");
-			
-			if (InBuildCookRun->IsUsingZenStore())
-			{
-				UATCommand += TEXT(" -zenstore");
-			}
+				TStringBuilder<64> InheritableCookOptions;
+				FCommandLine::BuildSubprocessCommandLine(ECommandLineArgumentFlags::CommandletContext, bOnlyInherited, InheritableCookOptions);
 
-			TStringBuilder<64> InheritableCookOptions;
-			FCommandLine::BuildSubprocessCommandLine(ECommandLineArgumentFlags::CommandletContext, bOnlyInherited, InheritableCookOptions);
+				FString AdditionalOptions = InBuildCookRun->GetCookOptions();
+				if (!AdditionalOptions.IsEmpty() || (InheritableCookOptions.Len() > 0))
+				{
+					UATCommand += TEXT(" -additionalcookeroptions=\"");
 
-			FString AdditionalOptions = InBuildCookRun->GetCookOptions();
-			if (!AdditionalOptions.IsEmpty() || (InheritableCookOptions.Len() > 0))
-			{
-				UATCommand += TEXT(" -additionalcookeroptions=\"");
+					// Escape any quotes in the argument list
+					UATCommand += AdditionalOptions.Replace(TEXT("\""), TEXT("\\\""));
 
-				// Escape any quotes in the argument list
-				UATCommand += AdditionalOptions.Replace(TEXT("\""), TEXT("\\\""));
+					if (InheritableCookOptions.Len() > 0)
+					{
+						if (!AdditionalOptions.IsEmpty())
+						{
+							UATCommand += TEXT(" ");
+						}
+						FString EscapedInheritableCookOptions(InheritableCookOptions);
+						EscapedInheritableCookOptions.ReplaceInline(TEXT("\""), TEXT("\\\""));
+						UATCommand += EscapedInheritableCookOptions;
+					}
 
-				if (InheritableCookOptions.Len() > 0)
-				{
-					if (!AdditionalOptions.IsEmpty())
+					// If the additional options ends with a slash, make sure we don't escape the quote
+					if (UATCommand.EndsWith("\\"))
 					{
-						UATCommand += TEXT(" ");
+						UATCommand += TEXT("\\");
 					}
-					FString EscapedInheritableCookOptions(InheritableCookOptions);
-					EscapedInheritableCookOptions.ReplaceInline(TEXT("\""), TEXT("\\\""));
-					UATCommand += EscapedInheritableCookOptions;
+
+					UATCommand += TEXT("\"");
 				}
 
-				// If the additional options ends with a slash, make sure we don't escape the quote
-				if (UATCommand.EndsWith("\\"))
+				if (FDerivedDataCacheInterface* DDC = GetDerivedDataCache())
 				{
-					UATCommand += TEXT("\\");
+					UATCommand += FString::Printf(TEXT(" -ddc=%s"), DDC->GetGraphName());
 				}
 
-				UATCommand += TEXT("\"");
-			}
-
-			if (FDerivedDataCacheInterface* DDC = GetDerivedDataCache())
-			{
-				UATCommand += FString::Printf(TEXT(" -ddc=%s"), DDC->GetGraphName());
-			}
+				//if UAT doesn't stick around as long as the process we are going to run, then we can't kill the COTF server when UAT goes down because the program
+				//will still need it.  If UAT DOES stick around with the process then we DO want the COTF server to die with UAT so the next time we launch we don't end up
+				//with two COTF servers.
+				if (bUATClosesAfterLaunch)
+				{
+					UATCommand += " -nokill";
+				}
+				UATCommand += MapList;
 
-			//if UAT doesn't stick around as long as the process we are going to run, then we can't kill the COTF server when UAT goes down because the program
-			//will still need it.  If UAT DOES stick around with the process then we DO want the COTF server to die with UAT so the next time we launch we don't end up
-			//with two COTF servers.
-			if (bUATClosesAfterLaunch)
-			{
-				UATCommand += " -nokill";
+				FCommandDesc Desc;
+				FText Command = LOCTEXT("LauncherCookOnTheFlyDesc", "Starting cook on the fly server");
+				Desc.Name = "Cook Server Task";
+				Desc.Desc = Command.ToString();
+				Desc.EndText = TEXT("********** COOK COMMAND COMPLETED **********");
+				OutCommands.Add(Desc);
+				if (CommandStart.Len() == 0)
+				{
+					CommandStart = TEXT("********** COOK COMMAND STARTED **********");
+				}
 			}
+			break;
+		case ELauncherProfileCookModes::OnTheFlyInEditor:
 			UATCommand += MapList;
-
-			FCommandDesc Desc;
-			FText Command = LOCTEXT("LauncherCookOnTheFlyDesc", "Starting cook on the fly server");
-			Desc.Name = "Cook Server Task";
-			Desc.Desc = Command.ToString();
-			Desc.EndText = TEXT("********** COOK COMMAND COMPLETED **********");
-			OutCommands.Add(Desc);
-			if (CommandStart.Len() == 0)
+			UATCommand += " -skipcook -cookonthefly -CookInEditor";
+			if (InBuildCookRun->IsUsingZenStore())
 			{
-				CommandStart = TEXT("********** COOK COMMAND STARTED **********");
+				UATCommand += TEXT(" -zenstore");
 			}
-		}
-		break;
-	case ELauncherProfileCookModes::OnTheFlyInEditor:
-		UATCommand += MapList;
-		UATCommand += " -skipcook -cookonthefly -CookInEditor";
-		if (InBuildCookRun->IsUsingZenStore())
-		{
-			UATCommand += TEXT(" -zenstore");
-		}
-		break;
-	case ELauncherProfileCookModes::ByTheBookInEditor:
-		UATCommand += MapList;
-		UATCommand += CultureList;
-		UATCommand += TEXT(" -skipcook -CookInEditor"); // don't cook anything the editor is doing it ;)
-		if (InBuildCookRun->IsUsingZenStore())
-		{
-			UATCommand += TEXT(" -zenstore");
-		}
-		if (InBuildCookRun->IsPackingWithUnrealPak())
-		{
-			UATCommand += TEXT(" -pak");
-			if (InBuildCookRun->IsUsingIoStore())
+			break;
+		case ELauncherProfileCookModes::ByTheBookInEditor:
+			UATCommand += MapList;
+			UATCommand += CultureList;
+			UATCommand += TEXT(" -skipcook -CookInEditor"); // don't cook anything the editor is doing it ;)
+			if (InBuildCookRun->IsUsingZenStore())
 			{
-				UATCommand += TEXT(" -iostore");
+				UATCommand += TEXT(" -zenstore");
 			}
-			if (InBuildCookRun->IsCompressed())
+			if (InBuildCookRun->IsPackingWithUnrealPak())
 			{
-				UATCommand += TEXT(" -compressed");
+				UATCommand += TEXT(" -pak");
+				if (InBuildCookRun->IsUsingIoStore())
+				{
+					UATCommand += TEXT(" -iostore");
+				}
+				if (InBuildCookRun->IsCompressed())
+				{
+					UATCommand += TEXT(" -compressed");
+				}
 			}
-		}
-		break;
-	case ELauncherProfileCookModes::DoNotCook:
-		UATCommand += TEXT(" -skipcook");
-		if (InBuildCookRun->IsUsingZenPakStreaming())
-		{
-			UATCommand += FString::Printf(TEXT(" -skippak -ZenWorkspaceSharePath=\\\"%s\\\" "), *InBuildCookRun->GetZenPakStreamingPath() );
-		}
-		else if (InBuildCookRun->IsImportingZenSnapshot() && InBuildCookRun->IsPackingWithUnrealPak())
-		{
-			UATCommand += TEXT(" -pak");
-			if (InBuildCookRun->IsUsingIoStore())
+			break;
+		case ELauncherProfileCookModes::DoNotCook:
+			UATCommand += TEXT(" -skipcook");
+			if (InBuildCookRun->IsUsingZenPakStreaming())
 			{
-				UATCommand += TEXT(" -iostore");
+				UATCommand += FString::Printf(TEXT(" -skippak -ZenWorkspaceSharePath=\\\"%s\\\" "), *InBuildCookRun->GetZenPakStreamingPath() );
 			}
-			if (InBuildCookRun->IsCompressed())
+			else if (InBuildCookRun->IsImportingZenSnapshot() && InBuildCookRun->IsPackingWithUnrealPak())
 			{
-				UATCommand += TEXT(" -compressed");
+				UATCommand += TEXT(" -pak");
+				if (InBuildCookRun->IsUsingIoStore())
+				{
+					UATCommand += TEXT(" -iostore");
+				}
+				if (InBuildCookRun->IsCompressed())
+				{
+					UATCommand += TEXT(" -compressed");
+				}
 			}
+			break;
 		}
-		break;
-	}
 
-	if ( InBuildCookRun->IsForDistribution() )
-	{
-		UATCommand += TEXT(" -distribution");
-	}
+		if ( InBuildCookRun->IsForDistribution() )
+		{
+			UATCommand += TEXT(" -distribution");
+		}
 
-	if (InBuildCookRun->IsCookingIncrementally())
-	{
-		UATCommand += TEXT(" -iterativecooking");
-	}
+		if (InBuildCookRun->IsCookingIncrementally())
+		{
+			UATCommand += TEXT(" -iterativecooking");
+		}
 
-	if ( InBuildCookRun->IsIterateSharedCookedBuild() )
-	{
-		UATCommand += TEXT(" -iteratesharedcookedbuild=usesyncedbuild");
-	}
+		if ( InBuildCookRun->IsIterateSharedCookedBuild() )
+		{
+			UATCommand += TEXT(" -iteratesharedcookedbuild=usesyncedbuild");
+		}
 
-	if (InBuildCookRun->GetSkipCookingEditorContent())
-	{
-		UATCommand += TEXT(" -SkipCookingEditorContent");
+		if (InBuildCookRun->GetSkipCookingEditorContent())
+		{
+			UATCommand += TEXT(" -SkipCookingEditorContent");
+		}
 	}
 
 	FString StageAdditionalCommandLine;
@@ -1077,6 +1082,11 @@ FString FLauncherWorker::CreateUATCommandForBuildCookRun( const ILauncherProfile
 	}
 
 	// stage/package/deploy
+	if (bPreStaged)
+	{
+		UATCommand += TEXT(" -prestage");
+	}
+
 	if (InBuildCookRun->GetDeploymentMode() != ELauncherProfileDeploymentModes::DoNotDeploy)
 	{
 		switch (InBuildCookRun->GetDeploymentMode())
@@ -1600,6 +1610,7 @@ ILauncherProfileUATCommandPtr FLauncherWorker::MakeTemporaryUnifiedRunCommand(co
 	UnifiedRunCommand->SetLaunchMode(ELauncherProfileLaunchModes::DefaultRole);
 	UnifiedRunCommand->SetAdditionalUATCommandLine(TEXT(" -skipbuild -skipcook -skipstage -skipdeploy")); // @todo: check if any of these are implied by the items above
 
+	bool bFirstItem = true;
 	for ( const ILauncherProfileBuildCookRunRef& BuildCookRun : BuildCookRuns)
 	{
 		// unify platforms & targets
@@ -1668,7 +1679,22 @@ ILauncherProfileUATCommandPtr FLauncherWorker::MakeTemporaryUnifiedRunCommand(co
 			UnifiedRunCommand->GetDefaultLaunchRole()->SetInitialMap(BuildCookRun->GetDefaultLaunchRole()->GetInitialMap());
 		}
 
+		// get consensus on using a pre-staged build
+		if (bFirstItem && BuildCookRun->IsUsingPreStagedBuild())
+		{
+			UnifiedRunCommand->SetUsePreStagedBuild(true);
+		}
+		else if (!bFirstItem && BuildCookRun->IsUsingPreStagedBuild() != UnifiedRunCommand->IsUsingPreStagedBuild())
+		{
+			// @todo: surface this better
+			UE_LOG(LogLauncherProfile, Warning, TEXT("Cannot create unified build cook run because not all builds agree on whether to use a pre-staged build"));
+			return nullptr;
+		}
+
 		// ... @todo: any other properties that need unification? ...
+
+
+		bFirstItem = false;
 	}
 
 	return UnifiedRunCommand;
diff --git a/Engine/Source/Developer/LauncherServices/Private/Profiles/LauncherProfile.cpp b/Engine/Source/Developer/LauncherServices/Private/Profiles/LauncherProfile.cpp
--- a/Engine/Source/Developer/LauncherServices/Private/Profiles/LauncherProfile.cpp
+++ b/Engine/Source/Developer/LauncherServices/Private/Profiles/LauncherProfile.cpp
@@ -79,7 +79,7 @@ FString LexToStringLocalized(ELauncherProfileValidationErrors::Type Value)
 	switch (Value)
 	{
 		case ELauncherProfileValidationErrors::CopyToDeviceRequiresCookByTheBook:
-			return LOCTEXT("CopyToDeviceRequiresCookByTheBookError", "Deployment by copying to device requires cooking or zen snapshot import.").ToString();
+			return LOCTEXT("CopyToDeviceRequiresCookByTheBookError", "Deployment by copying to device requires cooking, zen streaming or a pre-staged build.").ToString();
 		case ELauncherProfileValidationErrors::CustomRolesNotSupportedYet:
 			return LOCTEXT("CustomRolesNotSupportedYet", "Custom launch roles are not supported yet.").ToString();
 		case ELauncherProfileValidationErrors::DeployedDeviceGroupRequired:
diff --git a/Engine/Source/Developer/LauncherServices/Private/Profiles/LauncherProfile.h b/Engine/Source/Developer/LauncherServices/Private/Profiles/LauncherProfile.h
--- a/Engine/Source/Developer/LauncherServices/Private/Profiles/LauncherProfile.h
+++ b/Engine/Source/Developer/LauncherServices/Private/Profiles/LauncherProfile.h
@@ -2566,7 +2566,7 @@ protected:
 		}
 
 		// Deploy: deployment by copying to devices requires cooking by the book unless the data is available from zen
-		if ((BuildCookRun->GetDeploymentMode() == ELauncherProfileDeploymentModes::CopyToDevice) && ((BuildCookRun->GetCookMode() != ELauncherProfileCookModes::ByTheBook)&&(BuildCookRun->GetCookMode()!=ELauncherProfileCookModes::ByTheBookInEditor)) && !BuildCookRun->IsUsingZenPakStreaming() && !BuildCookRun->IsUsingZenStreaming())
+		if ((BuildCookRun->GetDeploymentMode() == ELauncherProfileDeploymentModes::CopyToDevice) && ((BuildCookRun->GetCookMode() != ELauncherProfileCookModes::ByTheBook)&&(BuildCookRun->GetCookMode()!=ELauncherProfileCookModes::ByTheBookInEditor)) && !BuildCookRun->IsUsingZenPakStreaming() && !BuildCookRun->IsUsingZenStreaming() && !BuildCookRun->IsUsingPreStagedBuild())
 		{
 			ValidationData.Errors.Add(ELauncherProfileValidationErrors::CopyToDeviceRequiresCookByTheBook);
 		}
diff --git a/Engine/Source/Developer/LauncherServices/Private/Profiles/LauncherProfileBuildCookRun.h b/Engine/Source/Developer/LauncherServices/Private/Profiles/LauncherProfileBuildCookRun.h
--- a/Engine/Source/Developer/LauncherServices/Private/Profiles/LauncherProfileBuildCookRun.h
+++ b/Engine/Source/Developer/LauncherServices/Private/Profiles/LauncherProfileBuildCookRun.h
@@ -783,6 +783,7 @@ public:
 		Writer.WriteValue("UseZenStreaming", bUseZenStreaming);
 		Writer.WriteValue("UseZenPakStreaming", bUseZenPakStreaming);
 		Writer.WriteValue("ZenPakStreamingPath", ZenPakStreamingPath);
+		Writer.WriteValue("PreStaged", bUsePreStagedBuild);
 		Writer.WriteValue("ClientArchitectures", ClientArchitectures);
 		Writer.WriteValue("ServerArchitectures", ServerArchitectures);
 		Writer.WriteValue("EditorArchitectures", EditorArchitectures);
@@ -1139,6 +1140,11 @@ public:
 			}
 		}
 
+		if (ProfileVersion >= LAUNCHERSERVICES_ADDPRESTAGED)
+		{
+			bUsePreStagedBuild = Object.GetBoolField(TEXT("PreStaged"));
+		}
+
 		// load the default launch role
 		TSharedPtr<FJsonObject> Role = Object.GetObjectField(TEXT("DefaultRole"));
 		DefaultLaunchRole->Load(*(Role.Get()));
@@ -1249,6 +1255,7 @@ public:
 		bImportingZenSnapshot = false;
 		bUseZenStreaming = false;
 		bUseZenPakStreaming = false;
+		bUsePreStagedBuild = false;
 		bIsDeviceASimulator = false;
 		bMakeBinaryConfig = false;
 
@@ -1714,6 +1721,21 @@ public:
 		return ZenPakStreamingPath;
 	}
 
+	virtual void SetUsePreStagedBuild( bool UsePreStagedBuild ) override
+	{
+		if (bUsePreStagedBuild != UsePreStagedBuild)
+		{
+			bUsePreStagedBuild = UsePreStagedBuild;
+			Validate();
+		}
+	}
+
+	virtual bool IsUsingPreStagedBuild() const override
+	{
+		return bUsePreStagedBuild;
+	}
+
+
 	virtual void SetDeviceIsASimulator(bool bInIsDeviceASimualtor) override
 	{
 		bIsDeviceASimulator = bInIsDeviceASimualtor;
@@ -2088,6 +2110,9 @@ private:
 	// Holds the full path to a folder to use for Zen pak streaming
 	FString ZenPakStreamingPath;
 
+	// Whether to use a pre-staged build
+	bool bUsePreStagedBuild;
+
     // Holds the time out time for the cook on the fly server
     uint32 Timeout;
 
diff --git a/Engine/Source/Developer/LauncherServices/Private/Profiles/LauncherProfileVersion.h b/Engine/Source/Developer/LauncherServices/Private/Profiles/LauncherProfileVersion.h
--- a/Engine/Source/Developer/LauncherServices/Private/Profiles/LauncherProfileVersion.h
+++ b/Engine/Source/Developer/LauncherServices/Private/Profiles/LauncherProfileVersion.h
@@ -44,6 +44,7 @@ enum ELauncherVersion
 	LAUNCHERSERVICES_ADDUATCOMMANDS = 45,
 	LAUNCHERSERVICES_INNERBUILDCOOKRUN = 46,
 	LAUNCHERSERVICES_ADDZENSTREAMINGOPTIONS = 47,
+	LAUNCHERSERVICES_ADDPRESTAGED = 48,
 	//ADD NEW STUFF HERE
 
 
diff --git a/Engine/Source/Developer/LauncherServices/Public/ILauncherProfileBuildCookRun.h b/Engine/Source/Developer/LauncherServices/Public/ILauncherProfileBuildCookRun.h
--- a/Engine/Source/Developer/LauncherServices/Public/ILauncherProfileBuildCookRun.h
+++ b/Engine/Source/Developer/LauncherServices/Public/ILauncherProfileBuildCookRun.h
@@ -1152,6 +1152,20 @@ public:
 	 */
 	virtual FString GetZenPakStreamingPath() const = 0;
 
+	/**
+	 * Set whether to use a pre-staged build
+	 * 
+	 * @param UsePreStagedBuild Whether use a pre-staged build
+	 */
+	virtual void SetUsePreStagedBuild( bool UsePreStagedBuild ) = 0;
+
+	/**
+	 * Checks whether a pre-staged build willbe used
+	 * 
+	 * @return true if using a pre-staged build, false otherwise
+	 */
+	virtual bool IsUsingPreStagedBuild() const = 0;
+
 	/**
 	 * Sets whether to make a binary config file during packaging
 	 * @param bMakeBinaryConfig Whether to make a binary config file during staging
