fix prestaged parameter

#rnx

#Synced-CL 51338024

diff --git a/Engine/Source/Developer/LauncherServices/Private/Launcher/LauncherWorker.cpp b/Engine/Source/Developer/LauncherServices/Private/Launcher/LauncherWorker.cpp
--- a/Engine/Source/Developer/LauncherServices/Private/Launcher/LauncherWorker.cpp
+++ b/Engine/Source/Developer/LauncherServices/Private/Launcher/LauncherWorker.cpp
@@ -1084,7 +1084,7 @@ FString FLauncherWorker::CreateUATCommandForBuildCookRun( const ILauncherProfile
 	// stage/package/deploy
 	if (bPreStaged)
 	{
-		UATCommand += TEXT(" -prestage");
+		UATCommand += TEXT(" -prestaged");
 	}
 
 	if (InBuildCookRun->GetDeploymentMode() != ELauncherProfileDeploymentModes::DoNotDeploy)
