Fixed Launcher Process and Thread leaks, and Memory Corruption

#Synced-CL 52646121

diff --git a/Engine/Source/Developer/LauncherServices/Private/Launcher/Launcher.cpp b/Engine/Source/Developer/LauncherServices/Private/Launcher/Launcher.cpp
--- a/Engine/Source/Developer/LauncherServices/Private/Launcher/Launcher.cpp
+++ b/Engine/Source/Developer/LauncherServices/Private/Launcher/Launcher.cpp
@@ -20,14 +20,11 @@ ILauncherWorkerPtr FLauncher::Launch(const TSharedRef<ITargetDeviceProxyManager>
 {
 	if (Profile->IsValidForLaunch())
 	{
-		FLauncherWorker* LauncherWorker = new FLauncherWorker(DeviceProxyManager, Profile);
-		FString WorkerName(FString::Printf(TEXT("LauncherWorker%i"), WorkerCounter.Increment()));
-
-		if ((LauncherWorker != nullptr) && (FRunnableThread::Create(LauncherWorker, *WorkerName) != nullptr))
+		TSharedPtr<FLauncherWorker> LauncherWorker = MakeShared<FLauncherWorker>(DeviceProxyManager, Profile, *FString::Printf(TEXT("LauncherWorker%i"), WorkerCounter.Increment()));
+		if (LauncherWorker->IsValid())
 		{
-			ILauncherWorkerPtr Worker = MakeShareable(LauncherWorker);
-			FLauncherWorkerStartedDelegate.Broadcast(Worker, Profile);
-			return Worker;
+			FLauncherWorkerStartedDelegate.Broadcast(LauncherWorker, Profile);
+			return MoveTemp(LauncherWorker);
 		}
 	}
 	else
diff --git a/Engine/Source/Developer/LauncherServices/Private/Launcher/LauncherTask.h b/Engine/Source/Developer/LauncherServices/Private/Launcher/LauncherTask.h
--- a/Engine/Source/Developer/LauncherServices/Private/Launcher/LauncherTask.h
+++ b/Engine/Source/Developer/LauncherServices/Private/Launcher/LauncherTask.h
@@ -23,12 +23,20 @@ public:
 	 * Creates and initializes a new instance.
 	 *
 	 * @param InName - The name of the task.
-	 * @param InDependency - An optional task that this task depends on.
+	 * @param InDesc - The description of the task.
 	 */
-	FLauncherTask( const FString& InName, const FString& InDesc)
-		: Name(InName)
-		, Desc(InDesc)
+	FLauncherTask(FString InName, FString InDesc)
+		: Name{MoveTemp(InName)}
+		, Desc{MoveTemp(InDesc)}
 	{ }
+	
+	virtual ~FLauncherTask() override
+	{
+		if (Thread)
+		{
+			Thread->WaitForCompletion();
+		}
+	}
 
 public:
 
@@ -52,7 +60,7 @@ public:
 	 *
 	 * @param ChainState - State data of the task chain.
 	 */
-	void Execute( const FLauncherTaskChainState& ChainState )
+	void Execute(const FLauncherTaskChainState* ChainState)
 	{
 		check(Status == ELauncherTaskStatus::Pending);
 
@@ -62,7 +70,7 @@ public:
 		FString TaskName(TEXT("FLauncherTask"));
 		TaskName.AppendInt(TaskCounter.Increment());
 		Status = ELauncherTaskStatus::Busy;
-		Thread = FRunnableThread::Create(this, *TaskName);
+		Thread = TUniquePtr<FRunnableThread>{FRunnableThread::Create(this, *TaskName)};
 	}
 
 	/**
@@ -152,7 +160,7 @@ public:
 		FPlatformProcess::Sleep(0.5f);
 
 		TaskStarted.Broadcast(Name);
-		bool bSucceeded = PerformTask(LocalChainState);
+		bool bSucceeded = PerformTask(*LocalChainState);
 		
 		if (bSucceeded)
 		{
@@ -278,7 +286,7 @@ protected:
 	 *
 	 * @return true if the task completed successfully, false otherwise.
 	 */
-	virtual bool PerformTask( FLauncherTaskChainState& ChainState ) = 0;
+	virtual bool PerformTask(const FLauncherTaskChainState& ChainState) = 0;
 
 protected:
 
@@ -313,7 +321,7 @@ private:
 	FDateTime EndTime;
 
 	// Holds the local state of this task chain.
-	FLauncherTaskChainState LocalChainState;
+	const FLauncherTaskChainState* LocalChainState = nullptr;
 
 	// Holds the task's name.
 	FString Name;
@@ -331,7 +339,7 @@ private:
 	bool bCancelling = false;
 
 	// Holds the thread that's running this task.
-	FRunnableThread* Thread = nullptr;
+	TUniquePtr<FRunnableThread> Thread;
 
 	// message delegates
 	FOnTaskStartedDelegate TaskStarted;
diff --git a/Engine/Source/Developer/LauncherServices/Private/Launcher/LauncherUATTask.h b/Engine/Source/Developer/LauncherServices/Private/Launcher/LauncherUATTask.h
--- a/Engine/Source/Developer/LauncherServices/Private/Launcher/LauncherUATTask.h
+++ b/Engine/Source/Developer/LauncherServices/Private/Launcher/LauncherUATTask.h
@@ -17,14 +17,12 @@ class FLauncherUATTask
 {
 public:
 
-	FLauncherUATTask( const FString& InCommandLine, const FString& InName, const FString& InDesc, void* InReadPipe, void* InWritePipe, FProcHandle& InProcessHandle, ILauncherWorker* InWorker, const FString& InCommandEnd)
-		: FLauncherTask(InName, InDesc)
-
-		, CommandLine(InCommandLine)
-		, ReadPipe(InReadPipe)
-		, WritePipe(InWritePipe)
-		, ProcessHandle(InProcessHandle)
-		, CommandText(InCommandEnd)
+	FLauncherUATTask(FString InCommandLine, FString InName, FString InDesc, UE::HAL::FWriteHandle InWritePipe, UE::HAL::FProcess& InProcess, ILauncherWorker* InWorker, FString InCommandEnd)
+		: FLauncherTask{MoveTemp(InName), MoveTemp(InDesc)}
+		, CommandLine{MoveTemp(InCommandLine)}
+		, WritePipe{InWritePipe}
+		, Process{InProcess}
+		, CommandText{MoveTemp(InCommandEnd)}
 	{
 		InWorker->OnOutputReceived().AddRaw(this, &FLauncherUATTask::HandleOutputReceived);
 	}
@@ -38,11 +36,12 @@ protected:
 	 *
 	 * @return true if the task completed successfully, false otherwise.
 	 */
-	virtual bool PerformTask( FLauncherTaskChainState& ChainState ) override
+	virtual bool PerformTask(const FLauncherTaskChainState& ChainState) override
 	{
 		// spawn a UAT process to cook the data
+		// UAT executable path
+		FString ExecutablePath = FPaths::ConvertRelativePathToFull(FPaths::EngineDir() + TEXTVIEW("Build/BatchFiles"));
 		// UAT executable
-		FString ExecutablePath = FPaths::ConvertRelativePathToFull(FPaths::EngineDir() + FString(TEXT("Build")) / TEXT("BatchFiles"));
 #if PLATFORM_MAC
 		FString Executable = TEXT("RunUAT.command");
 #elif PLATFORM_LINUX
@@ -67,21 +66,21 @@ protected:
 		UATCommandLine += TEXT(" ");
 		UATCommandLine += CommandLine;
 
-
 		// launch UAT and monitor its progress
-		ProcessHandle = FPlatformProcess::CreateProc(*(ExecutablePath / Executable), *UATCommandLine, false, true, true, NULL, 0, *ExecutablePath, WritePipe);
+		Process = UE::HAL::FProcess{{.Uri = *(ExecutablePath / Executable), .Arguments = *UATCommandLine, .WorkingDirectory = *ExecutablePath, .bHidden = true, .StdOut = WritePipe}};
 
-		while (FPlatformProcess::IsProcRunning(ProcessHandle) && !EndTextFound)
+		while (!EndTextFound.load(std::memory_order_relaxed))
 		{
-			FPlatformProcess::Sleep(0.25);
-		}
+			if (!Process.IsRunning())
+			{
+				Result = *Process.GetExitCode();
+				break;
+			}
 
-		if (!EndTextFound && !FPlatformProcess::GetProcReturnCode(ProcessHandle, &Result))
-		{
-			return false;
+			FPlatformProcess::Sleep(0.25);
 		}
 
-		return (Result == 0);
+		return Result == 0;
 	}
 
 	void HandleOutputReceived(const FString& InMessage)
@@ -95,7 +94,7 @@ protected:
 			++WarningCounter;
 		}
 
-		EndTextFound |= InMessage.Contains(CommandText);
+		EndTextFound.store(EndTextFound.load(std::memory_order_relaxed) | InMessage.Contains(CommandText), std::memory_order_relaxed);
 	}
 
 private:
@@ -103,16 +102,15 @@ private:
 	// Command line
 	FString CommandLine;
 
-	// read and write pipe
-	void* ReadPipe = nullptr;
-	void* WritePipe = nullptr;
+	// write pipe
+	UE::HAL::FWriteHandle WritePipe;
 
-	// process handle
-	FProcHandle& ProcessHandle;
+	// process
+	UE::HAL::FProcess& Process;
 
 	// The end text marker for this UAT command
 	FString CommandText;
 
-	bool EndTextFound = false;
+	std::atomic_bool EndTextFound = false;
 };
 
diff --git a/Engine/Source/Developer/LauncherServices/Private/Launcher/LauncherWorker.cpp b/Engine/Source/Developer/LauncherServices/Private/Launcher/LauncherWorker.cpp
--- a/Engine/Source/Developer/LauncherServices/Private/Launcher/LauncherWorker.cpp
+++ b/Engine/Source/Developer/LauncherServices/Private/Launcher/LauncherWorker.cpp
@@ -34,11 +34,13 @@ FThreadSafeCounter FLauncherTask::TaskCounter;
 /* FLauncherWorker structors
  *****************************************************************************/
 
-FLauncherWorker::FLauncherWorker(const TSharedRef<ITargetDeviceProxyManager>& InDeviceProxyManager, const ILauncherProfileRef& InProfile)
+FLauncherWorker::FLauncherWorker(const TSharedRef<ITargetDeviceProxyManager>& InDeviceProxyManager, const ILauncherProfileRef& InProfile, const TCHAR* InWorkerName)
 	: DeviceProxyManager(InDeviceProxyManager)
 	, Profile(InProfile)
 {
 	CreateAndExecuteTasks(InProfile);
+	
+	Thread = TUniquePtr<FRunnableThread>{FRunnableThread::Create(this, InWorkerName)};
 }
 
 
@@ -80,13 +82,13 @@ uint32 FLauncherWorker::Run( )
 		FString Line;
 		auto ReadLines = [this, &MessageReceived, &Line]
 		{
-			FString NewLine = FPlatformProcess::ReadPipe(ReadPipe);
+			FString NewLine = InputPipe.Read();
 			if (NewLine.IsEmpty())
 			{
 				return false;
 			}
 
-			// process the string to break it up in to lines
+			// process the string to break it up into lines
 			Line += NewLine;
 	
 			if (int32 NewLineIndex; NewLine.FindChar(TCHAR('\n'), NewLineIndex))
@@ -108,7 +110,7 @@ uint32 FLauncherWorker::Run( )
 		};
 
 		// wait for tasks to be completed
-		while (Status == ELauncherWorkerStatus::Busy)
+		while (Status.load(std::memory_order_relaxed) == ELauncherWorkerStatus::Busy)
 		{
 			FPlatformProcess::Sleep(0.01f);
 
@@ -116,7 +118,7 @@ uint32 FLauncherWorker::Run( )
 
 			if (TaskChain.IsValid() && TaskChain->IsChainFinished())
 			{
-				Status = ELauncherWorkerStatus::Completed;
+				Status.store(ELauncherWorkerStatus::Completed, std::memory_order_relaxed);
 
 				while (ReadLines());
 
@@ -131,10 +133,14 @@ uint32 FLauncherWorker::Run( )
 	}
 
 	// wait for tasks to be canceled
-	if (Status == ELauncherWorkerStatus::Canceling)
+	if (Status.load(std::memory_order_relaxed) == ELauncherWorkerStatus::Canceling)
 	{
 		// kill the uat process tree
-		FPlatformProcess::TerminateProc(ProcHandle, true);
+		if (Process)
+		{
+			Process.Kill(true);
+		}
+		
 		// kill any lingering target processes left after killing uat
 		TerminateLaunchedProcess();
 
@@ -146,12 +152,13 @@ uint32 FLauncherWorker::Run( )
 		}		
 	}
 
-	FPlatformProcess::ClosePipe(ReadPipe, WritePipe);
+	Process = nullptr;
+	InputPipe = nullptr;
 
-	if (!TaskChain.IsValid() || Status == ELauncherWorkerStatus::Canceling)
+	if (!TaskChain.IsValid() || Status.load(std::memory_order_relaxed) == ELauncherWorkerStatus::Canceling)
 	{
 		LaunchCanceled.Broadcast(FPlatformTime::Seconds() - LaunchStartTime);
-		Status = ELauncherWorkerStatus::Canceled;
+		Status.store(ELauncherWorkerStatus::Canceled, std::memory_order_relaxed);
 	}
 	else
 	{
@@ -179,19 +186,19 @@ void FLauncherWorker::Stop( )
 
 void FLauncherWorker::Cancel( )
 {
-	if (Status == ELauncherWorkerStatus::Busy)
+	if (Status.load(std::memory_order_relaxed) == ELauncherWorkerStatus::Busy)
 	{
-		Status = ELauncherWorkerStatus::Canceling;
+		Status.store(ELauncherWorkerStatus::Canceling, std::memory_order_relaxed);
 	}
 }
 
 
 void FLauncherWorker::CancelAndWait( )
 {
-	if (Status == ELauncherWorkerStatus::Busy)
+	if (Status.load(std::memory_order_relaxed) == ELauncherWorkerStatus::Busy)
 	{
-		Status = ELauncherWorkerStatus::Canceling;
-		while (Status != ELauncherWorkerStatus::Canceled)
+		Status.store(ELauncherWorkerStatus::Canceling, std::memory_order_relaxed);
+		while (Status.load(std::memory_order_relaxed) != ELauncherWorkerStatus::Canceled)
 		{
 			FPlatformProcess::Sleep(0);
 		}
@@ -1761,7 +1768,7 @@ FString FLauncherWorker::FinalizeCommandLine( const FString& CommandLine ) const
 
 void FLauncherWorker::CreateAndExecuteTasks( const ILauncherProfileRef& InProfile )
 {
-	FPlatformProcess::CreatePipe(ReadPipe, WritePipe);
+	InputPipe = UE::HAL::NewPipe;
 
 	TArray<FString> Platforms;
 	TArray<FString> DeviceIDs;
@@ -1807,19 +1814,21 @@ void FLauncherWorker::CreateAndExecuteTasks( const ILauncherProfileRef& InProfil
 	if (InProfile->GetFirstBuildCookRun().IsValid() && InProfile->GetFirstBuildCookRun()->GetCookMode() == ELauncherProfileCookModes::ByTheBookInEditor)
 	{
 		// need a command which will wait for the cook to finish
-		class FWaitForCookInEditorToFinish : public FLauncherTask
+		struct FWaitForCookInEditorToFinish: FLauncherTask
 		{
-		public:
-			FWaitForCookInEditorToFinish() : FLauncherTask( FString(TEXT("Cooking in the editor")), FString(TEXT("Preparing content to run on device")))
+			FWaitForCookInEditorToFinish()
+				: FLauncherTask{TEXT("Cooking in the editor"), TEXT("Preparing content to run on device")}
 			{
 			}
-			virtual bool PerformTask( FLauncherTaskChainState& ChainState ) override
+
+		protected:
+			virtual bool PerformTask(const FLauncherTaskChainState& InChainState) override
 			{
-				while ( !ChainState.Profile->OnIsCookFinished().Execute() )
+				while ( !InChainState.Profile->OnIsCookFinished().Execute() )
 				{
 					if (IsCancelling())
 					{
-						ChainState.Profile->OnCookCanceled().Execute();
+						InChainState.Profile->OnCookCanceled().Execute();
 						return false;
 					}
 					FPlatformProcess::Sleep( 0.1f );
@@ -1827,8 +1836,7 @@ void FLauncherWorker::CreateAndExecuteTasks( const ILauncherProfileRef& InProfil
 				return true;
 			}
 		};
-		TSharedPtr<FLauncherTask> WaitTask = MakeShareable(new FWaitForCookInEditorToFinish());
-		AddTask(WaitTask);
+		AddTask(MakeShared<FWaitForCookInEditorToFinish>());
 	}
 
 	// gather all ordered UAT commands
@@ -1907,19 +1915,17 @@ void FLauncherWorker::CreateAndExecuteTasks( const ILauncherProfileRef& InProfil
 		Commands.Add(Desc);
 	}
 
-
-	TSharedPtr<FLauncherTask> LaunchTask = MakeShareable(new FLauncherUATTask(UATCommandString, TEXT("Launch Task"), TEXT("Launching UAT..."), ReadPipe, WritePipe, ProcHandle, this, StartString));
+	TSharedPtr<FLauncherTask> LaunchTask = MakeShareable(new FLauncherUATTask(UATCommandString, TEXT("Launch Task"), TEXT("Launching UAT..."), InputPipe, Process, this, StartString));
 	AddTask(LaunchTask, UnifiedRunUATCommand);
 	for (int32 Index = 0; Index < Commands.Num(); ++Index)
 	{
-		class FLauncherWaitTask : public FLauncherTask
+		struct FLauncherWaitTask: FLauncherTask
 		{
-		public:
-			FLauncherWaitTask( const FString& InCommandEnd, const FString& InName, const FString& InDesc, FProcHandle& InProcessHandle, ILauncherWorker* InWorker)
-				: FLauncherTask(InName, InDesc)
-				, CommandText(InCommandEnd)
-				, ProcessHandle(InProcessHandle)
-				, LauncherWorker(InWorker)
+			FLauncherWaitTask(FString InCommandEnd, FString InName, FString InDesc, UE::HAL::FProcess& InProcess, ILauncherWorker* InWorker)
+				: FLauncherTask{MoveTemp(InName), MoveTemp(InDesc)}
+				, CommandText{MoveTemp(InCommandEnd)}
+				, Process{InProcess}
+				, LauncherWorker{InWorker}
 			{
 				InWorker->OnOutputReceived().AddRaw(this, &FLauncherWaitTask::HandleOutputReceived);
 			}
@@ -1930,40 +1936,40 @@ void FLauncherWorker::CreateAndExecuteTasks( const ILauncherProfileRef& InProfil
 			}
 
 		protected:
-			virtual bool PerformTask( FLauncherTaskChainState& ChainState ) override
+			virtual bool PerformTask(const FLauncherTaskChainState& InChainState) override
 			{
-				while (FPlatformProcess::IsProcRunning(ProcessHandle) && !EndTextFound)
+				while (!EndTextFound.load(std::memory_order_relaxed))
 				{
+					if (!Process.IsRunning())
+					{
+						Result = *Process.GetExitCode();
+						break;
+					}
+
 					FPlatformProcess::Sleep(0.25);
 				}
-				if (!EndTextFound && !FPlatformProcess::GetProcReturnCode(ProcessHandle, &Result))
-				{
-					return false;
-				}
-				return (Result == 0);
+
+				return Result == 0;
 			}
 
 			void HandleOutputReceived(const FString& InMessage)
 			{
-				EndTextFound |= InMessage.Contains(CommandText);
+				EndTextFound.store(EndTextFound.load(std::memory_order_relaxed) | InMessage.Contains(CommandText), std::memory_order_relaxed);
 			}
 
 		private:
 			FString CommandText;
-			FProcHandle& ProcessHandle;
+			UE::HAL::FProcess& Process;
 			ILauncherWorker* LauncherWorker = nullptr;
-			bool EndTextFound = false;
+			std::atomic_bool EndTextFound = false;
 		};			
-
-		TSharedPtr<FLauncherTask> WaitTask = MakeShareable(new FLauncherWaitTask(Commands[Index].EndText, Commands[Index].Name, Commands[Index].Desc, ProcHandle, this));
-		AddTask(WaitTask);
+		AddTask(MakeShared<FLauncherWaitTask>(Commands[Index].EndText, Commands[Index].Name, Commands[Index].Desc, Process, this));
 	}
 
 	// execute the chain
-	FLauncherTaskChainState ChainState;
 	ChainState.Profile = InProfile;
 	ChainState.SessionId = FGuid::NewGuid();
-	TaskChain->Execute(ChainState);
+	TaskChain->Execute(&ChainState);
 }
 
 /** Callback for when a device proxy has been removed from the device proxy manager. */
diff --git a/Engine/Source/Developer/LauncherServices/Private/Launcher/LauncherWorker.h b/Engine/Source/Developer/LauncherServices/Private/Launcher/LauncherWorker.h
--- a/Engine/Source/Developer/LauncherServices/Private/Launcher/LauncherWorker.h
+++ b/Engine/Source/Developer/LauncherServices/Private/Launcher/LauncherWorker.h
@@ -4,8 +4,10 @@
 
 #include "CoreMinimal.h"
 #include "ILauncherWorker.h"
+#include "LauncherTaskChainState.h"
 #include "HAL/PlatformProcess.h"
 #include "HAL/Runnable.h"
+#include "HAL/RunnableThread.h"
 
 class FLauncherTask;
 class ITargetDeviceProxyManager;
@@ -23,16 +25,28 @@ public:
 	/**
 	 * Creates and initializes a new instance.
 	 *
-	 * @param InProfile The profile to process.
 	 * @param InDeviceProxyManager The target device proxy manager to use.
+	 * @param InProfile The profile to process.
+	 * @param InWorkerName The name of the worker thread.
 	 */
-	FLauncherWorker(const TSharedRef<ITargetDeviceProxyManager>& InDeviceProxyManager, const ILauncherProfileRef& InProfile);
+	FLauncherWorker(const TSharedRef<ITargetDeviceProxyManager>& InDeviceProxyManager, const ILauncherProfileRef& InProfile, const TCHAR* InWorkerName);
 
-	~FLauncherWorker()
+	virtual ~FLauncherWorker() override
 	{
+		if (Thread)
+		{
+			Thread->WaitForCompletion();
+		}
+
 		TaskChain.Reset();
 		DisableDeviceDiscoveryListener();
 	}
+	
+	/** Returns whether the worker is valid to run. */
+	bool IsValid() const
+	{
+		return bool(Thread);
+	}
 
 public:
 
@@ -50,9 +64,9 @@ public:
 
 	virtual void CancelAndWait( ) override;
 
-	virtual ELauncherWorkerStatus::Type GetStatus( ) const  override
+	virtual ELauncherWorkerStatus::Type GetStatus() const override
 	{
-		return Status;
+		return Status.load(std::memory_order_relaxed);
 	}
 
 	virtual int32 GetTasks( TArray<ILauncherTaskPtr>& OutTasks ) const override;
@@ -155,15 +169,17 @@ private:
 	ILauncherProfilePtr Profile;
 
 	// Holds the worker's current status.
-	volatile ELauncherWorkerStatus::Type Status = ELauncherWorkerStatus::Busy;
+	std::atomic<ELauncherWorkerStatus::Type> Status = ELauncherWorkerStatus::Busy;
 
 	// Holds the first task in the task chain.
 	TSharedPtr<FLauncherTask> TaskChain;
+	FLauncherTaskChainState ChainState;
+
+	TUniquePtr<FRunnableThread> Thread;
 
-	// holds the read and write pipes and the running UAT process
-	void* ReadPipe = nullptr;
-	void* WritePipe = nullptr;
-	FProcHandle ProcHandle;
+	// Holds the input pipe and the running UAT process.
+	UE::HAL::FInputPipe InputPipe;
+	UE::HAL::FProcess Process;
 
 	double StageStartTime = 0.0;
 	double LaunchStartTime = 0.0;
